Added initial support for the SteamVR API.  Plugin is disabled by default and is in the Experimental plugins folder.

Requires the Steamworks SDK to build.  If the SDK is not available, a stub version of the plugin is built instead.  This is primarily an issue for builds from github.

diff --git a/Engine/Plugins/Experimental/SteamVR/Source/SteamVR/Private/SteamVRHMD.cpp b/Engine/Plugins/Experimental/SteamVR/Source/SteamVR/Private/SteamVRHMD.cpp
new file mode 100644
--- /dev/null
+++ b/Engine/Plugins/Experimental/SteamVR/Source/SteamVR/Private/SteamVRHMD.cpp
@@ -0,0 +1,716 @@
+// Copyright 1998-2014 Epic Games, Inc. All Rights Reserved.
+
+#include "SteamVRPrivatePCH.h"
+#include "SteamVRHMD.h"
+
+#include "RendererPrivate.h"
+#include "ScenePrivate.h"
+#include "PostProcess/PostProcessHMD.h"
+
+
+//---------------------------------------------------
+// SteamVR Plugin Implementation
+//---------------------------------------------------
+
+class FSteamVRPlugin : public ISteamVRPlugin
+{
+	/** IHeadMountedDisplayModule implementation */
+	virtual TSharedPtr< class IHeadMountedDisplay > CreateHeadMountedDisplay() override;
+};
+
+IMPLEMENT_MODULE( FSteamVRPlugin, SteamVR )
+
+TSharedPtr< class IHeadMountedDisplay > FSteamVRPlugin::CreateHeadMountedDisplay()
+{
+#if STEAMVR_SUPPORTED_PLATFORMS
+	TSharedPtr< FSteamVRHMD > SteamVRHMD( new FSteamVRHMD() );
+	if( SteamVRHMD->IsInitialized() )
+	{
+		return SteamVRHMD;
+	}
+#endif//STEAMVR_SUPPORTED_PLATFORMS
+	return nullptr;
+}
+
+
+//---------------------------------------------------
+// SteamVR IHeadMountedDisplay Implementation
+//---------------------------------------------------
+
+#if STEAMVR_SUPPORTED_PLATFORMS
+
+
+bool FSteamVRHMD::IsHMDEnabled() const
+{
+	return bHmdEnabled;
+}
+
+void FSteamVRHMD::EnableHMD(bool enable)
+{
+	bHmdEnabled = enable;
+
+	if (!bHmdEnabled)
+	{
+		EnableStereo(false);
+	}
+}
+
+EHMDDeviceType::Type FSteamVRHMD::GetHMDDeviceType() const
+{
+	return EHMDDeviceType::DT_OculusRift;		//@todo steamvr: steamvr post
+}
+
+bool FSteamVRHMD::GetHMDMonitorInfo(MonitorInfo& MonitorDesc) const
+{
+	if (IsInitialized())
+	{
+		int32 X, Y;
+		uint32 Width, Height;
+		Hmd->GetWindowBounds(&X, &Y, &Width, &Height);
+
+		MonitorDesc.MonitorName = DisplayId;
+		MonitorDesc.MonitorId	= 0;
+		MonitorDesc.DesktopX	= X;
+		MonitorDesc.DesktopY	= Y;
+		MonitorDesc.ResolutionX = Width;
+		MonitorDesc.ResolutionY = Height;
+
+		return true;
+	}
+	else
+	{
+		MonitorDesc.MonitorName = "";
+		MonitorDesc.MonitorId = 0;
+		MonitorDesc.DesktopX = MonitorDesc.DesktopY = MonitorDesc.ResolutionX = MonitorDesc.ResolutionY = 0;
+	}
+
+	return false;
+}
+
+void FSteamVRHMD::GetFieldOfView(float& OutHFOVInDegrees, float& OutVFOVInDegrees) const
+{
+	OutHFOVInDegrees = 0.0f;
+	OutVFOVInDegrees = 0.0f;
+}
+
+bool FSteamVRHMD::DoesSupportPositionalTracking() const
+{
+	return true;
+}
+
+bool FSteamVRHMD::HasValidTrackingPosition() const
+{
+	return bHmdPosTracking && bHaveVisionTracking;
+}
+
+void FSteamVRHMD::GetPositionalTrackingCameraProperties(FVector& OutOrigin, FRotator& OutOrientation, float& OutHFOV, float& OutVFOV, float& OutCameraDistance, float& OutNearPlane, float& OutFarPlane) const
+{
+}
+
+void FSteamVRHMD::SetInterpupillaryDistance(float NewInterpupillaryDistance)
+{
+}
+
+float FSteamVRHMD::GetInterpupillaryDistance() const
+{
+	return 0.064f;
+}
+
+void FSteamVRHMD::GetCurrentPose(FQuat& CurrentOrientation, FVector& CurrentPosition, float MotionPredictionInSeconds) const
+{
+	vr::HmdMatrix34_t Pose;
+	vr::HmdTrackingResult Result;
+	bool bValid = Hmd->GetTrackerFromHeadPose(MotionPredictionInSeconds, &Pose, &Result);
+
+	bHaveVisionTracking = bValid && (Result == vr::TrackingResult_Running_OK);
+
+	if (!bHaveVisionTracking)
+	{
+		return;
+	}
+
+	PoseToOrientationAndPosition(Pose, CurrentOrientation, CurrentPosition);
+}
+
+void FSteamVRHMD::PoseToOrientationAndPosition(const vr::HmdMatrix34_t& InPose, FQuat& OutOrientation, FVector& OutPosition) const
+{
+	FMatrix Pose = ToFMatrix(InPose);
+	FQuat Orientation(Pose);
+
+	OutOrientation.X = -Orientation.Z;
+	OutOrientation.Y = Orientation.X;
+	OutOrientation.Z = Orientation.Y;
+	OutOrientation.W = -Orientation.W;
+	OutPosition = FVector(-Pose.M[3][2], Pose.M[3][0], Pose.M[3][1]) * WorldToMetersScale;
+}
+
+namespace
+{
+	void UpdatePlayerViewPoint(const FQuat& CurrentOrientation, const FVector& CurrentPosition, FSceneView& View)
+	{
+		//if (!CurrentOrientation.IsNormalized() || !View.BaseHmdOrientation.IsNormalized())
+		//{
+		//	return;
+		//}
+
+		//// Given ViewRotation is BaseOrientation * BaseHmdOrientation
+		//FQuat BaseOrientation = (View.ViewRotation.Quaternion() * View.BaseHmdOrientation.Inverse());
+		//BaseOrientation.Normalize();
+
+		//// Use the inverse of the controller viewpoint to find the view-based offset that the HMD is currently in
+		//// ViewLocation = BaseLocation + (BaseOrientation * BaseHmdLocation)
+		//FVector HeadOffset = BaseOrientation.RotateVector(CurrentPosition - View.BaseHmdLocation);
+
+		//// Update new view rotation and location
+		//View.ViewRotation = (BaseOrientation * CurrentOrientation).Rotator().GetNormalized();
+		//View.ViewLocation += HeadOffset;// *View.WorldToMetersScale;
+
+		//View.UpdateViewMatrix();
+		//GetViewFrustumBounds(View.ViewFrustum, View.ViewMatrices.GetViewProjMatrix(), false);
+	}
+}
+
+
+void FSteamVRHMD::GetCurrentOrientationAndPosition(FQuat& CurrentOrientation, FVector& CurrentPosition) const
+{
+	checkf(IsInGameThread());
+	GetCurrentPose(CurHmdOrientation, CurHmdPosition, MotionPredictionInSecondsGame);
+	CurrentOrientation = LastHmdOrientation = CurHmdOrientation;
+
+	CurrentPosition = CurHmdPosition;
+}
+
+ISceneViewExtension* FSteamVRHMD::GetViewExtension()
+{
+	return this;
+}
+
+void FSteamVRHMD::ApplyHmdRotation(APlayerController* PC, FRotator& ViewRotation)
+{
+	ViewRotation.Normalize();
+
+	GetCurrentPose(CurHmdOrientation, CurHmdPosition, MotionPredictionInSecondsGame);
+	LastHmdOrientation = CurHmdOrientation;
+
+	const FRotator DeltaRot = ViewRotation - PC->GetControlRotation();
+	DeltaControlRotation = (DeltaControlRotation + DeltaRot).GetNormalized();
+
+	// Pitch from other sources is never good, because there is an absolute up and down that must be respected to avoid motion sickness.
+	// Same with roll.
+	DeltaControlRotation.Pitch = 0;
+	DeltaControlRotation.Roll = 0;
+	DeltaControlOrientation = DeltaControlRotation.Quaternion();
+
+	ViewRotation = FRotator(DeltaControlOrientation * CurHmdOrientation);
+}
+
+void FSteamVRHMD::UpdatePlayerCameraRotation(APlayerCameraManager* Camera, struct FMinimalViewInfo& POV)
+{
+	GetCurrentPose(CurHmdOrientation, CurHmdPosition, MotionPredictionInSecondsGame);
+	LastHmdOrientation = CurHmdOrientation;
+
+	DeltaControlRotation = POV.Rotation;
+	DeltaControlOrientation = DeltaControlRotation.Quaternion();
+
+	// Apply HMD orientation to camera rotation.
+	POV.Rotation = FRotator(POV.Rotation.Quaternion() * CurHmdOrientation);
+}
+
+bool FSteamVRHMD::IsChromaAbCorrectionEnabled() const
+{
+	return true;
+}
+
+bool FSteamVRHMD::Exec( UWorld* InWorld, const TCHAR* Cmd, FOutputDevice& Ar )
+{
+	if (FParse::Command( &Cmd, TEXT("STEREO") ))
+	{
+		if (FParse::Command(&Cmd, TEXT("ON")))
+		{
+			if (!IsHMDEnabled())
+			{
+				Ar.Logf(TEXT("HMD is disabled. Use 'hmd enable' to re-enable it."));
+			}
+			EnableStereo(true);
+			return true;
+		}
+		else if (FParse::Command(&Cmd, TEXT("OFF")))
+		{
+			EnableStereo(false);
+			return true;
+		}
+
+		float val;
+		if (FParse::Value(Cmd, TEXT("E="), val))
+		{
+			IPD = val;
+		}
+	}
+	else if (FParse::Command(&Cmd, TEXT("HMD")))
+	{
+		if (FParse::Command(&Cmd, TEXT("ENABLE")))
+		{
+			EnableHMD(true);
+			return true;
+		}
+		else if (FParse::Command(&Cmd, TEXT("DISABLE")))
+		{
+			EnableHMD(false);
+			return true;
+		}
+
+		float val;
+		if (FParse::Value(Cmd, TEXT("MPG="), val))
+		{
+			MotionPredictionInSecondsGame = val;
+			return true;
+		}
+
+		if (FParse::Value(Cmd, TEXT("MPR="), val))
+		{
+			MotionPredictionInSecondsRender = val;
+			return true;
+		}
+	}
+	else if (FParse::Command(&Cmd, TEXT("UNCAPFPS")))
+	{
+		GEngine->bSmoothFrameRate = false;
+		return true;
+	}
+
+	return false;
+}
+
+void FSteamVRHMD::OnScreenModeChange(EWindowMode::Type WindowMode)
+{
+	EnableStereo(WindowMode != EWindowMode::Windowed);
+}
+
+bool FSteamVRHMD::IsPositionalTrackingEnabled() const
+{
+	return bHmdPosTracking;
+}
+
+bool FSteamVRHMD::EnablePositionalTracking(bool enable)
+{
+	bHmdPosTracking = enable;
+	return IsPositionalTrackingEnabled();
+}
+
+bool FSteamVRHMD::IsHeadTrackingAllowed() const
+{
+	return GEngine->IsStereoscopic3D();
+}
+
+bool FSteamVRHMD::IsInLowPersistenceMode() const
+{
+	return true;
+}
+
+void FSteamVRHMD::EnableLowPersistenceMode(bool Enable)
+{
+}
+
+
+void FSteamVRHMD::ResetOrientationAndPosition(float yaw)
+{
+	//@todo steamvr: ResetOrientationAndPosition()
+}
+
+void FSteamVRHMD::UpdateScreenSettings(const FViewport*)
+{
+}
+
+bool FSteamVRHMD::IsStereoEnabled() const
+{
+	return bStereoEnabled && bHmdEnabled;
+}
+
+bool FSteamVRHMD::EnableStereo(bool stereo)
+{
+	bStereoEnabled = (IsHMDEnabled()) ? stereo : false;
+	return bStereoEnabled;
+}
+
+void FSteamVRHMD::AdjustViewRect(EStereoscopicPass StereoPass, int32& X, int32& Y, uint32& SizeX, uint32& SizeY) const
+{
+	//@todo steamvr: get the actual rects from steamvr
+	SizeX = SizeX / 2;
+	if( StereoPass == eSSP_RIGHT_EYE )
+	{
+		X += SizeX;
+	}
+}
+
+void FSteamVRHMD::CalculateStereoViewOffset(const enum EStereoscopicPass StereoPassType, const FRotator& ViewRotation, const float WorldToMeters, FVector& ViewLocation)
+{
+	if( StereoPassType != eSSP_FULL)
+	{
+		vr::Hmd_Eye HmdEye = (StereoPassType == eSSP_LEFT_EYE) ? vr::Eye_Left : vr::Eye_Right;
+		vr::HmdMatrix34_t HeadFromEye = Hmd->GetHeadFromEyePose(HmdEye);
+
+		// grab the eye position, currently ignoring the rotation supplied by GetHeadFromEyePose()
+		FVector TotalOffset = FVector(-HeadFromEye.m[2][3], HeadFromEye.m[0][3], HeadFromEye.m[1][3]) * WorldToMeters;
+
+		ViewLocation += ViewRotation.Quaternion().RotateVector(TotalOffset);
+
+		const FVector vHMDPosition = DeltaControlOrientation.RotateVector(CurHmdPosition);
+		ViewLocation += vHMDPosition;
+		LastHmdPosition = CurHmdPosition;
+	}
+}
+
+FMatrix FSteamVRHMD::GetStereoProjectionMatrix(const enum EStereoscopicPass StereoPassType, const float FOV) const
+{
+	check(IsStereoEnabled());
+
+	vr::Hmd_Eye HmdEye = (StereoPassType == eSSP_LEFT_EYE) ? vr::Eye_Left : vr::Eye_Right;
+	float Left, Right, Top, Bottom;
+
+	Hmd->GetProjectionRaw(HmdEye, &Right, &Left, &Top, &Bottom);
+	Bottom *= -1.0f;
+	Top *= -1.0f;
+	Right *= -1.0f;
+	Left *= -1.0f;
+
+	float ZNear = GNearClippingPlane;
+
+	float SumRL = (Right + Left);
+	float SumTB = (Top + Bottom);
+	float InvRL = (1.0f / (Right - Left));
+	float InvTB = (1.0f / (Top - Bottom));
+
+#if 1
+	FMatrix Mat = FMatrix(
+		FPlane((2.0f * InvRL), 0.0f, 0.0f, 0.0f),
+		FPlane(0.0f, (2.0f * InvTB), 0.0f, 0.0f),
+		FPlane((SumRL * InvRL), (SumTB * InvTB), 0.0f, 1.0f),
+		FPlane(0.0f, 0.0f, ZNear, 0.0f)
+		);
+#else
+	vr::HmdMatrix44_t SteamMat = Hmd->GetProjectionMatrix(HmdEye, ZNear, 10000.0f, vr::API_DirectX);
+	FMatrix Mat = ToFMatrix(SteamMat);
+
+	Mat.M[3][3] = 0.0f;
+	Mat.M[2][3] = 1.0f;
+	Mat.M[2][2] = 0.0f;
+	Mat.M[3][2] = ZNear;
+#endif
+
+	return Mat;
+
+}
+
+void FSteamVRHMD::InitCanvasFromView(FSceneView* InView, UCanvas* Canvas)
+{
+}
+
+void FSteamVRHMD::PushViewportCanvas(EStereoscopicPass StereoPass, FCanvas *InCanvas, UCanvas *InCanvasObject, FViewport *InViewport) const 
+{
+	FMatrix m;
+	m.SetIdentity();
+	InCanvas->PushAbsoluteTransform(m);
+}
+
+void FSteamVRHMD::PushViewCanvas(EStereoscopicPass StereoPass, FCanvas *InCanvas, UCanvas *InCanvasObject, FSceneView *InView) const 
+{
+	FMatrix m;
+	m.SetIdentity();
+	InCanvas->PushAbsoluteTransform(m);
+}
+
+void FSteamVRHMD::GetEyeRenderParams_RenderThread(EStereoscopicPass StereoPass, FVector2D& EyeToSrcUVScaleValue, FVector2D& EyeToSrcUVOffsetValue) const
+{
+	if (StereoPass == eSSP_LEFT_EYE)
+	{
+		EyeToSrcUVOffsetValue.X = 0.0f;
+		EyeToSrcUVOffsetValue.Y = 0.0f;
+
+		EyeToSrcUVScaleValue.X = 0.5f;
+		EyeToSrcUVScaleValue.Y = 1.0f;
+	}
+	else
+	{
+		EyeToSrcUVOffsetValue.X = 0.5f;
+		EyeToSrcUVOffsetValue.Y = 0.0f;
+
+		EyeToSrcUVScaleValue.X = 0.5f;
+		EyeToSrcUVScaleValue.Y = 1.0f;
+	}
+}
+
+
+void FSteamVRHMD::ModifyShowFlags(FEngineShowFlags& ShowFlags)
+{
+	ShowFlags.MotionBlur = 0;
+	ShowFlags.HMDDistortion = true;
+	ShowFlags.ScreenPercentage = 1.0f;
+	ShowFlags.StereoRendering = IsStereoEnabled();
+}
+
+void FSteamVRHMD::SetupView(FSceneViewFamily& InViewFamily, FSceneView& InView)
+{
+	InView.BaseHmdOrientation = FQuat(FRotator(0.0f,0.0f,0.0f));
+	InView.BaseHmdLocation = FVector(0.f);
+	WorldToMetersScale = InView.WorldToMetersScale;
+	InViewFamily.bUseSeparateRenderTarget = false;
+}
+
+void FSteamVRHMD::PreRenderView_RenderThread(FSceneView& View)
+{
+	check(IsInRenderingThread());
+
+	FVector CurrentPosition;
+	FQuat CurrentOrientation;
+
+	// Update again for low-latency
+	GetCurrentPose(CurrentOrientation, CurrentPosition, MotionPredictionInSecondsRender);
+
+//	UpdatePlayerViewPoint(CurrentOrientation, CurrentPosition, View);
+}
+
+void FSteamVRHMD::PreRenderViewFamily_RenderThread(FSceneViewFamily& ViewFamily)
+{
+	check(IsInRenderingThread());
+}
+
+
+FSteamVRHMD::FSteamVRHMD() :
+	Hmd(nullptr),
+	bHmdEnabled(true),
+	bStereoEnabled(false),
+	bHmdPosTracking(true),
+	bHaveVisionTracking(false),
+	IPD(0.064f),
+	LastHmdOrientation(FQuat::Identity),
+	CurHmdOrientation(FQuat::Identity),
+	DeltaControlRotation(FRotator::ZeroRotator),
+	DeltaControlOrientation(FQuat::Identity),
+	CurHmdPosition(FVector::ZeroVector),
+	WorldToMetersScale(100.0f),
+	MotionPredictionInSecondsGame(0.0f),
+	MotionPredictionInSecondsRender(0.0f),
+	SteamDLLHandle(nullptr),
+	IdealScreenPercentage(100.0f)
+{
+	Startup();
+}
+
+FSteamVRHMD::~FSteamVRHMD()
+{
+	Shutdown();
+}
+
+bool FSteamVRHMD::IsInitialized() const
+{
+	return (Hmd != nullptr);
+}
+
+void FSteamVRHMD::Startup()
+{
+	// load the steam library
+	if (!LoadSteamModule())
+	{
+		return;
+	}
+
+	bool ret = SteamAPI_Init();
+
+	// look for a headset
+	vr::HmdError HmdErr = vr::HmdError_None;
+	Hmd = vr::VR_Init(&HmdErr);
+
+	if (Hmd != nullptr)
+	{
+		// grab info about the attached display
+		char Buf[128];
+		FString DriverId;
+		if (Hmd->GetDriverId(Buf, 128) > 0)
+		{
+			DriverId = FString(UTF8_TO_TCHAR(Buf));
+		}
+
+		if (Hmd->GetDisplayId(Buf, 128) > 0)
+		{
+			DisplayId = FString(UTF8_TO_TCHAR(Buf));
+		}
+
+		// create our distortion meshes
+		DistortionMesh[0].CreateMesh(Hmd, vr::Eye_Left);
+		DistortionMesh[1].CreateMesh(Hmd, vr::Eye_Right);
+
+		// determine our ideal screen percentage
+		uint32 RecommendedWidth, RecommendedHeight;
+		Hmd->GetRecommendedRenderTargetSize(&RecommendedWidth, &RecommendedHeight);
+		RecommendedWidth *= 2;
+
+		int32 ScreenX, ScreenY;
+		uint32 ScreenWidth, ScreenHeight;
+		Hmd->GetWindowBounds(&ScreenX, &ScreenY, &ScreenWidth, &ScreenHeight);
+
+		float WidthPercentage = ((float)RecommendedWidth / (float)ScreenWidth) * 100.0f;
+		float HeightPercentage = ((float)RecommendedHeight / (float)ScreenHeight) * 100.0f;
+
+		float IdealScreenPercentage = FMath::Max(WidthPercentage, HeightPercentage);
+
+		//@todo steamvr: move out of here
+		IConsoleVariable* CScrPercVar = IConsoleManager::Get().FindConsoleVariable(TEXT("r.ScreenPercentage"));
+
+		if (FMath::RoundToInt(CScrPercVar->GetFloat()) != FMath::RoundToInt(IdealScreenPercentage))
+		{
+			CScrPercVar->Set(IdealScreenPercentage);
+		}
+
+		// enable vsync
+		static IConsoleVariable* CVSyncVar = IConsoleManager::Get().FindConsoleVariable(TEXT("r.VSync"));
+		CVSyncVar->Set(true);
+
+		//@todo steamvr: better prediction
+		MotionPredictionInSecondsGame = 2.0f * (1.0f / 75.0f);
+		MotionPredictionInSecondsRender = 1.0f * (1.0f / 75.0f);
+
+		// Uncap fps to enable FPS higher than 62
+		GEngine->bSmoothFrameRate = false;
+		UE_LOG(LogHMD, Log, TEXT("SteamVR initialized.  Driver: %s  Display: %s"), *DriverId, *DisplayId);
+	}
+	else
+	{
+		UE_LOG(LogHMD, Warning, TEXT("SteamVR failed to initialize.  Error: %d"), (int32)HmdErr);
+	}
+}
+
+void FSteamVRHMD::Shutdown()
+{
+	// unload steam library
+	UnloadSteamModule();
+
+	// shut down our headset
+	vr::VR_Shutdown();
+	Hmd = nullptr;
+}
+
+bool FSteamVRHMD::LoadSteamModule()
+{
+#if PLATFORM_WINDOWS
+	#if PLATFORM_64BITS
+		FString RootSteamPath = FPaths::EngineDir() / FString::Printf(TEXT("Binaries/ThirdParty/Steamworks/%s/Win64/"), STEAM_SDK_VER); 
+		FPlatformProcess::PushDllDirectory(*RootSteamPath);
+		SteamDLLHandle = FPlatformProcess::GetDllHandle(*(RootSteamPath + "steam_api64.dll"));
+		FPlatformProcess::PopDllDirectory(*RootSteamPath);
+	#else
+		FString RootSteamPath = FPaths::EngineDir() / FString::Printf(TEXT("Binaries/ThirdParty/Steamworks/%s/Win32/"), STEAM_SDK_VER); 
+		FPlatformProcess::PushDllDirectory(*RootSteamPath);
+		SteamDLLHandle = FPlatformProcess::GetDllHandle(*(RootSteamPath + "steam_api.dll"));
+		FPlatformProcess::PopDllDirectory(*RootSteamPath);
+	#endif
+#elif PLATFORM_MAC
+	SteamDLLHandle = FPlatformProcess::GetDllHandle(TEXT("libsteam_api.dylib"));
+#endif	//PLATFORM_WINDOWS
+
+	if (!SteamDLLHandle)
+	{
+		UE_LOG(LogHMD, Warning, TEXT("Failed to load Steam library."));
+		return false;
+	}
+
+	return true;
+}
+
+void FSteamVRHMD::UnloadSteamModule()
+{
+	if (SteamDLLHandle != nullptr)
+	{
+		FPlatformProcess::FreeDllHandle(SteamDLLHandle);
+		SteamDLLHandle = nullptr;
+	}
+}
+
+
+FSteamVRDistortionMesh::FSteamVRDistortionMesh() :
+	Verts(nullptr),
+	Indices(nullptr),
+	NumVerts(0),
+	NumIndices(0),
+	NumTriangles(0)
+{
+}
+
+FSteamVRDistortionMesh::~FSteamVRDistortionMesh()
+{
+	if (Verts != nullptr)
+	{
+		delete[] Verts;
+	}
+
+	if (Indices != nullptr)
+	{
+		delete[] Indices;
+	}
+}
+
+void FSteamVRDistortionMesh::CreateMesh(vr::IHmd* Hmd, vr::Hmd_Eye Eye)
+{
+	static const uint32 ResU = 32;
+	static const uint32 ResV = 32;
+
+	static const uint32 ResUMinus1 = ResU - 1;
+	static const uint32 ResVMinus1 = ResV - 1;
+
+#if 0
+	// should be taking Hmd viewport into account
+	uint32 VPX, VPY, VPWidth, VPHeight;
+	Hmd->GetEyeOutputViewport(vr::Eye_Left, &VPX, &VPY, &VPWidth, &VPHeight);
+	Hmd->GetEyeOutputViewport(vr::Eye_Right, &VPX, &VPY, &VPWidth, &VPHeight);
+#endif
+
+	NumVerts = ResU * ResV;
+	NumTriangles = 2 * ResUMinus1 * ResVMinus1;
+	NumIndices = 3 * NumTriangles;
+
+	Verts = new FDistortionVertex[NumVerts];
+	Indices = new uint16[NumIndices];
+
+	uint32 TriIndex = 0;
+
+	for (uint32 IndexV = 0; IndexV < ResV; ++IndexV)
+	{
+		float V = (float)IndexV / (float)ResVMinus1;
+		float Y = FMath::Lerp(1.0f, -1.0f, V);
+
+		for (uint32 IndexU = 0; IndexU < ResU; ++IndexU)
+		{
+			float U = (float)IndexU / (float)ResUMinus1;
+			float X = (Eye == vr::Eye_Left) ? FMath::Lerp(-1.0f, 0.0f, U) : FMath::Lerp(0.0f, 1.0f, U);
+
+			vr::DistortionCoordinates_t DistCoords = Hmd->ComputeDistortion(Eye, U, V);
+
+			uint32 VertIndex = IndexV * ResU + IndexU;
+
+			Verts[VertIndex].Position = FVector2D(X, Y);
+			Verts[VertIndex].TexR = FVector2D(DistCoords.rfRed[0], DistCoords.rfRed[1]);
+			Verts[VertIndex].TexG = FVector2D(DistCoords.rfGreen[0], DistCoords.rfGreen[1]);
+			Verts[VertIndex].TexB = FVector2D(DistCoords.rfBlue[0], DistCoords.rfBlue[1]);
+			Verts[VertIndex].TimewarpFactor = 0.0f;
+			Verts[VertIndex].VignetteFactor = 1.0f;
+
+			if ((IndexV > 0) && (IndexU > 0))
+			{
+				uint16 A = (IndexV - 1) * ResU + (IndexU - 1);
+				uint16 B = (IndexV - 1) * ResU + IndexU;
+				uint16 C = IndexV * ResU + (IndexU - 1);
+				uint16 D = IndexV * ResU + IndexU;
+
+				Indices[TriIndex++] = A;
+				Indices[TriIndex++] = B;
+				Indices[TriIndex++] = D;
+
+				Indices[TriIndex++] = A;
+				Indices[TriIndex++] = D;
+				Indices[TriIndex++] = C;
+			}
+		}
+	}
+}
+
+
+#endif //STEAMVR_SUPPORTED_PLATFORMS
diff --git a/Engine/Plugins/Experimental/SteamVR/Source/SteamVR/Private/SteamVRHMD.h b/Engine/Plugins/Experimental/SteamVR/Source/SteamVR/Private/SteamVRHMD.h
new file mode 100644
--- /dev/null
+++ b/Engine/Plugins/Experimental/SteamVR/Source/SteamVR/Private/SteamVRHMD.h
@@ -0,0 +1,184 @@
+// Copyright 1998-2014 Epic Games, Inc. All Rights Reserved.
+
+#pragma once
+#include "ISteamVRPlugin.h"
+#include "HeadMountedDisplay.h"
+#include "IHeadMountedDisplay.h"
+
+#if STEAMVR_SUPPORTED_PLATFORMS
+
+#include "SceneViewExtension.h"
+
+
+class FSteamVRDistortionMesh
+{
+public:
+	struct FDistortionVertex*	Verts;
+	uint16*						Indices;
+	uint32						NumVerts;
+	uint32						NumIndices;
+	uint32						NumTriangles;
+
+	FSteamVRDistortionMesh();
+	~FSteamVRDistortionMesh();
+
+	void CreateMesh(vr::IHmd* Hmd, vr::Hmd_Eye Eye);
+};
+
+
+/**
+ * SteamVR Head Mounted Display
+ */
+class FSteamVRHMD : public IHeadMountedDisplay, public ISceneViewExtension
+{
+public:
+	/** IHeadMountedDisplay interface */
+	virtual bool IsHMDEnabled() const override;
+	virtual void EnableHMD(bool allow = true) override;
+	virtual EHMDDeviceType::Type GetHMDDeviceType() const override;
+	virtual bool GetHMDMonitorInfo(MonitorInfo&) const override;
+
+	virtual void GetFieldOfView(float& OutHFOVInDegrees, float& OutVFOVInDegrees) const override;
+
+	virtual bool DoesSupportPositionalTracking() const override;
+	virtual bool HasValidTrackingPosition() const override;
+	virtual void GetPositionalTrackingCameraProperties(FVector& OutOrigin, FRotator& OutOrientation, float& OutHFOV, float& OutVFOV, float& OutCameraDistance, float& OutNearPlane, float& OutFarPlane) const override;
+
+	virtual void SetInterpupillaryDistance(float NewInterpupillaryDistance) override;
+	virtual float GetInterpupillaryDistance() const override;
+
+	virtual void GetCurrentOrientationAndPosition(FQuat& CurrentOrientation, FVector& CurrentPosition) const override;
+	virtual class ISceneViewExtension* GetViewExtension() override;
+	virtual void ApplyHmdRotation(APlayerController* PC, FRotator& ViewRotation) override;
+	virtual void UpdatePlayerCameraRotation(APlayerCameraManager*, struct FMinimalViewInfo& POV) override;
+
+	virtual bool IsChromaAbCorrectionEnabled() const override;
+
+	virtual bool Exec( UWorld* InWorld, const TCHAR* Cmd, FOutputDevice& Ar ) override;
+	virtual void OnScreenModeChange(EWindowMode::Type WindowMode) override;
+
+	virtual bool IsPositionalTrackingEnabled() const override;
+	virtual bool EnablePositionalTracking(bool enable) override;
+
+	virtual bool IsHeadTrackingAllowed() const override;
+
+	virtual bool IsInLowPersistenceMode() const override;
+	virtual void EnableLowPersistenceMode(bool Enable = true) override;
+
+	virtual void ResetOrientationAndPosition(float yaw = 0.f) override;
+
+	virtual void DrawDistortionMesh_RenderThread(struct FRenderingCompositePassContext& Context, const FSceneView& View, const FIntPoint& TextureSize) override;
+	virtual void UpdateScreenSettings(const FViewport*) override;
+
+	/** IStereoRendering interface */
+	virtual bool IsStereoEnabled() const override;
+	virtual bool EnableStereo(bool stereo = true) override;
+	virtual void AdjustViewRect(EStereoscopicPass StereoPass, int32& X, int32& Y, uint32& SizeX, uint32& SizeY) const override;
+	virtual void CalculateStereoViewOffset(const EStereoscopicPass StereoPassType, const FRotator& ViewRotation, 
+										   const float MetersToWorld, FVector& ViewLocation) override;
+	virtual FMatrix GetStereoProjectionMatrix(const EStereoscopicPass StereoPassType, const float FOV) const override;
+	virtual void InitCanvasFromView(FSceneView* InView, UCanvas* Canvas) override;
+	virtual void PushViewportCanvas(EStereoscopicPass StereoPass, FCanvas *InCanvas, UCanvas *InCanvasObject, FViewport *InViewport) const override;
+	virtual void PushViewCanvas(EStereoscopicPass StereoPass, FCanvas *InCanvas, UCanvas *InCanvasObject, FSceneView *InView) const override;
+	virtual void GetEyeRenderParams_RenderThread(EStereoscopicPass StereoPass, FVector2D& EyeToSrcUVScaleValue, FVector2D& EyeToSrcUVOffsetValue) const override;
+
+	/** ISceneViewExtension interface */
+	virtual void ModifyShowFlags(FEngineShowFlags& ShowFlags) override;
+	virtual void SetupView(FSceneViewFamily& InViewFamily, FSceneView& InView) override;
+	virtual void PreRenderView_RenderThread(FSceneView& InView) override;
+	virtual void PreRenderViewFamily_RenderThread(FSceneViewFamily& InViewFamily) override;
+
+public:
+	/** Constructor */
+	FSteamVRHMD();
+
+	/** Destructor */
+	virtual ~FSteamVRHMD();
+
+	/** @return	True if the API was initialized OK */
+	bool IsInitialized() const;
+
+private:
+
+	/**
+	 * Starts up the SteamVR API
+	 */
+	void Startup();
+
+	/**
+	 * Shuts down the SteamVR API
+	 */
+	void Shutdown();
+
+	bool LoadSteamModule();
+	void UnloadSteamModule();
+
+	void PoseToOrientationAndPosition(const vr::HmdMatrix34_t& Pose, FQuat& OutOrientation, FVector& OutPosition) const;
+	void GetCurrentPose(FQuat& CurrentOrientation, FVector& CurrentPosition, float MotionPredictionInSeconds) const;
+
+
+	FORCEINLINE FMatrix ToFMatrix(const vr::HmdMatrix34_t& tm) const
+	{
+		// Rows and columns are swapped between vr::HmdMatrix34_t and FMatrix
+		return FMatrix(
+			FPlane(tm.m[0][0], tm.m[1][0], tm.m[2][0], 0.0f),
+			FPlane(tm.m[0][1], tm.m[1][1], tm.m[2][1], 0.0f),
+			FPlane(tm.m[0][2], tm.m[1][2], tm.m[2][2], 0.0f),
+			FPlane(tm.m[0][3], tm.m[1][3], tm.m[2][3], 1.0f));
+	}
+
+	FORCEINLINE FMatrix ToFMatrix(const vr::HmdMatrix44_t& tm) const
+	{
+		// Rows and columns are swapped between vr::HmdMatrix44_t and FMatrix
+		return FMatrix(
+			FPlane(tm.m[0][0], tm.m[1][0], tm.m[2][0], tm.m[3][0]),
+			FPlane(tm.m[0][1], tm.m[1][1], tm.m[2][1], tm.m[3][1]),
+			FPlane(tm.m[0][2], tm.m[1][2], tm.m[2][2], tm.m[3][2]),
+			FPlane(tm.m[0][3], tm.m[1][3], tm.m[2][3], tm.m[3][3]));
+	}
+
+private:
+
+	vr::IHmd* Hmd;
+	bool bHmdEnabled;
+	bool bStereoEnabled;
+	bool bHmdPosTracking;
+	mutable bool bHaveVisionTracking;
+
+	float IPD;
+
+	/** Player's orientation tracking */
+	mutable FQuat			CurHmdOrientation;
+
+	FRotator				DeltaControlRotation;    // same as DeltaControlOrientation but as rotator
+	FQuat					DeltaControlOrientation; // same as DeltaControlRotation but as quat
+
+	mutable FVector			CurHmdPosition;
+
+	mutable FQuat			LastHmdOrientation; // contains last APPLIED ON GT HMD orientation
+	FVector					LastHmdPosition;	// contains last APPLIED ON GT HMD position
+
+	/** HMD base values, specify forward orientation and zero pos offset */
+//	OVR::Vector3f			BaseOffset;			// base position, in SteamVR coords
+	FQuat					BaseOrientation;	// base orientation
+
+	/** World units (UU) to Meters scale.  Read from the level, and used to transform positional tracking data */
+	float WorldToMetersScale;
+
+	/** Motion prediction (in seconds). 0 - no prediction */
+	float MotionPredictionInSecondsGame;
+	float MotionPredictionInSecondsRender;
+
+	void* SteamDLLHandle;
+	
+	FString DisplayId;
+
+	float IdealScreenPercentage;
+
+	FSteamVRDistortionMesh DistortionMesh[2];
+};
+
+
+DEFINE_LOG_CATEGORY_STATIC(LogHMD, Log, All);
+
+#endif //STEAMVR_SUPPORTED_PLATFORMS
diff --git a/Engine/Plugins/Experimental/SteamVR/Source/SteamVR/Private/SteamVRPrivatePCH.h b/Engine/Plugins/Experimental/SteamVR/Source/SteamVR/Private/SteamVRPrivatePCH.h
new file mode 100644
--- /dev/null
+++ b/Engine/Plugins/Experimental/SteamVR/Source/SteamVR/Private/SteamVRPrivatePCH.h
@@ -0,0 +1,28 @@
+// Copyright 1998-2012 Epic Games, Inc. All Rights Reserved.
+
+#pragma once
+
+/** Name of the current Steam SDK version in use (matches directory name) */
+#define STEAM_SDK_VER TEXT("Steamv130")
+
+#include "Engine.h"
+#include "IHeadMountedDisplay.h"
+
+// @todo Steam: Steam headers trigger secure-C-runtime warnings in Visual C++. Rather than mess with _CRT_SECURE_NO_WARNINGS, we'll just
+//	disable the warnings locally. Remove when this is fixed in the SDK
+#if WITH_STEAMWORKS
+
+#ifdef _MSC_VER
+#pragma warning(push)
+#pragma warning(disable:4996)
+#endif
+
+#include "steam/steam_api.h"
+#include "steam/steamvr.h"
+
+// @todo Steam: See above
+#ifdef _MSC_VER
+#pragma warning(pop)
+#endif
+
+#endif // WITH_STEAMWORKS
\ No newline at end of file
diff --git a/Engine/Plugins/Experimental/SteamVR/Source/SteamVR/Private/SteamVRRender.cpp b/Engine/Plugins/Experimental/SteamVR/Source/SteamVR/Private/SteamVRRender.cpp
new file mode 100644
--- /dev/null
+++ b/Engine/Plugins/Experimental/SteamVR/Source/SteamVR/Private/SteamVRRender.cpp
@@ -0,0 +1,29 @@
+// Copyright 1998-2014 Epic Games, Inc. All Rights Reserved.
+//
+#include "SteamVRPrivatePCH.h"
+#include "SteamVRHMD.h"
+
+#include "RendererPrivate.h"
+#include "ScenePrivate.h"
+#include "PostProcess/PostProcessHMD.h"
+
+#if STEAMVR_SUPPORTED_PLATFORMS
+
+void FSteamVRHMD::DrawDistortionMesh_RenderThread(struct FRenderingCompositePassContext& Context, const FSceneView& View, const FIntPoint& TextureSize)
+{
+	float ClipSpaceQuadZ = 0.0f;
+	FMatrix QuadTexTransform = FMatrix::Identity;
+	FMatrix QuadPosTransform = FMatrix::Identity;
+	const FIntRect SrcRect = View.ViewRect;
+
+	FRHICommandListImmediate& RHICmdList = Context.RHICmdList;
+	const FSceneViewFamily& ViewFamily = *(View.Family);
+	FIntPoint ViewportSize = ViewFamily.RenderTarget->GetSizeXY();
+	RHICmdList.SetViewport(0, 0, 0.0f, ViewportSize.X, ViewportSize.Y, 1.0f);
+
+	FSteamVRDistortionMesh& CurMesh = DistortionMesh[(View.StereoPass == eSSP_LEFT_EYE) ? 0 : 1];
+
+	DrawIndexedPrimitiveUP(Context.RHICmdList, PT_TriangleList, 0, CurMesh.NumVerts, CurMesh.NumTriangles, CurMesh.Indices, sizeof(CurMesh.Indices[0]), CurMesh.Verts, sizeof(CurMesh.Verts[0]));
+}
+
+#endif // STEAMVR_SUPPORTED_PLATFORMS
\ No newline at end of file
diff --git a/Engine/Plugins/Experimental/SteamVR/Source/SteamVR/Public/ISteamVRPlugin.h b/Engine/Plugins/Experimental/SteamVR/Source/SteamVR/Public/ISteamVRPlugin.h
new file mode 100644
--- /dev/null
+++ b/Engine/Plugins/Experimental/SteamVR/Source/SteamVR/Public/ISteamVRPlugin.h
@@ -0,0 +1,40 @@
+// Copyright 1998-2014 Epic Games, Inc. All Rights Reserved.
+
+#pragma once
+
+#include "ModuleManager.h"
+#include "IHeadMountedDisplayModule.h"
+
+#define STEAMVR_SUPPORTED_PLATFORMS ((PLATFORM_WINDOWS || PLATFORM_MAC) && WITH_STEAMWORKS)
+
+/**
+ * The public interface to this module.  In most cases, this interface is only public to sibling modules 
+ * within this plugin.
+ */
+class ISteamVRPlugin : public IHeadMountedDisplayModule
+{
+
+public:
+
+	/**
+	 * Singleton-like access to this module's interface.  This is just for convenience!
+	 * Beware of calling this during the shutdown phase, though.  Your module might have been unloaded already.
+	 *
+	 * @return Returns singleton instance, loading the module on demand if needed
+	 */
+	static inline ISteamVRPlugin& Get()
+	{
+		return FModuleManager::LoadModuleChecked< ISteamVRPlugin >( "SteamVR" );
+	}
+
+	/**
+	 * Checks to see if this module is loaded and ready.  It is only valid to call Get() if IsAvailable() returns true.
+	 *
+	 * @return True if the module is loaded and ready to use
+	 */
+	static inline bool IsAvailable()
+	{
+		return FModuleManager::Get().IsModuleLoaded( "SteamVR" );
+	}
+};
+
diff --git a/Engine/Plugins/Experimental/SteamVR/Source/SteamVR/SteamVR.Build.cs b/Engine/Plugins/Experimental/SteamVR/Source/SteamVR/SteamVR.Build.cs
new file mode 100644
--- /dev/null
+++ b/Engine/Plugins/Experimental/SteamVR/Source/SteamVR/SteamVR.Build.cs
@@ -0,0 +1,41 @@
+// Copyright 1998-2014 Epic Games, Inc. All Rights Reserved.
+
+namespace UnrealBuildTool.Rules
+{
+	public class SteamVR : ModuleRules
+	{
+		public SteamVR(TargetInfo Target)
+		{
+			PrivateIncludePaths.AddRange(
+				new string[] {
+					"SteamVR/Private",
+					"../../../../Source/Runtime/Renderer/Private",
+					// ... add other private include paths required here ...
+				}
+				);
+
+			PrivateDependencyModuleNames.AddRange(
+				new string[]
+				{
+					"Core",
+					"CoreUObject",
+					"Engine",
+					"RHI",
+					"RenderCore",
+					"Renderer",
+					"ShaderCore",
+					"HeadMountedDisplay"
+				}
+				);
+
+			Definitions.Add("WITH_STEAMWORKS=" + (UEBuildConfiguration.bCompileSteamOSS ? "1" : "0"));
+			
+			// Currently, we only support SteamVR on Windows and Mac
+
+            if (UEBuildConfiguration.bCompileSteamOSS && (Target.Platform == UnrealTargetPlatform.Win32 || Target.Platform == UnrealTargetPlatform.Win64 || Target.Platform == UnrealTargetPlatform.Mac))
+            {
+				AddThirdPartyPrivateStaticDependencies(Target, "Steamworks");
+            }
+		}
+	}
+}
diff --git a/Engine/Plugins/Experimental/SteamVR/SteamVR.uplugin b/Engine/Plugins/Experimental/SteamVR/SteamVR.uplugin
new file mode 100644
--- /dev/null
+++ b/Engine/Plugins/Experimental/SteamVR/SteamVR.uplugin
@@ -0,0 +1,22 @@
+{
+	"FileVersion" : 3,
+
+	"FriendlyName" : "SteamVR Plugin",
+	"Version" : 1,
+	"VersionName" : "1.0",
+	"CreatedBy" : "Epic Games, Inc.",
+	"CreatedByURL" : "http://epicgames.com",
+	"EngineVersion" : "4.2.0",
+	"Description" : "[EARLY ACCESS PREVIEW] Support for the SteamVR API",
+	"Category" : "Head Mounted Displays.Devices",
+	"EnabledByDefault" : false,
+
+	"Modules" :
+	[
+		{
+			"Name" : "SteamVR",
+			"Type" : "Runtime",
+			"WhitelistPlatforms" : [ "Win64", "Win32", "Mac" ]
+		}
+	]
+}
