- Disabled SteamVR Mac support for now.

diff --git a/Engine/Plugins/Experimental/SteamVR/Source/SteamVR/Public/ISteamVRPlugin.h b/Engine/Plugins/Experimental/SteamVR/Source/SteamVR/Public/ISteamVRPlugin.h
--- a/Engine/Plugins/Experimental/SteamVR/Source/SteamVR/Public/ISteamVRPlugin.h
+++ b/Engine/Plugins/Experimental/SteamVR/Source/SteamVR/Public/ISteamVRPlugin.h
@@ -5,7 +5,7 @@
 #include "ModuleManager.h"
 #include "IHeadMountedDisplayModule.h"
 
-#define STEAMVR_SUPPORTED_PLATFORMS ((PLATFORM_WINDOWS || PLATFORM_MAC) && WITH_STEAMWORKS)
+#define STEAMVR_SUPPORTED_PLATFORMS (PLATFORM_WINDOWS && WITH_STEAMWORKS)
 
 /**
  * The public interface to this module.  In most cases, this interface is only public to sibling modules 
diff --git a/Engine/Plugins/Experimental/SteamVR/Source/SteamVR/SteamVR.Build.cs b/Engine/Plugins/Experimental/SteamVR/Source/SteamVR/SteamVR.Build.cs
--- a/Engine/Plugins/Experimental/SteamVR/Source/SteamVR/SteamVR.Build.cs
+++ b/Engine/Plugins/Experimental/SteamVR/Source/SteamVR/SteamVR.Build.cs
@@ -30,9 +30,7 @@ namespace UnrealBuildTool.Rules
 
 			Definitions.Add("WITH_STEAMWORKS=" + (UEBuildConfiguration.bCompileSteamOSS ? "1" : "0"));
 			
-			// Currently, we only support SteamVR on Windows and Mac
-
-            if (UEBuildConfiguration.bCompileSteamOSS && (Target.Platform == UnrealTargetPlatform.Win32 || Target.Platform == UnrealTargetPlatform.Win64 || Target.Platform == UnrealTargetPlatform.Mac))
+            if (UEBuildConfiguration.bCompileSteamOSS && (Target.Platform == UnrealTargetPlatform.Win32 || Target.Platform == UnrealTargetPlatform.Win64))
             {
 				AddThirdPartyPrivateStaticDependencies(Target, "Steamworks");
             }
diff --git a/Engine/Plugins/Experimental/SteamVR/SteamVR.uplugin b/Engine/Plugins/Experimental/SteamVR/SteamVR.uplugin
--- a/Engine/Plugins/Experimental/SteamVR/SteamVR.uplugin
+++ b/Engine/Plugins/Experimental/SteamVR/SteamVR.uplugin
@@ -16,7 +16,7 @@
 		{
 			"Name" : "SteamVR",
 			"Type" : "Runtime",
-			"WhitelistPlatforms" : [ "Win64", "Win32", "Mac" ]
+			"WhitelistPlatforms" : [ "Win64", "Win32" ]
 		}
 	]
 }
