refactor

diff --git a/Engine/Plugins/Experimental/SteamVR/Source/SteamVR/Private/SteamVRHMD.cpp b/Engine/Plugins/Experimental/SteamVR/Source/SteamVR/Private/SteamVRHMD.cpp
--- a/Engine/Plugins/Experimental/SteamVR/Source/SteamVR/Private/SteamVRHMD.cpp
+++ b/Engine/Plugins/Experimental/SteamVR/Source/SteamVR/Private/SteamVRHMD.cpp
@@ -559,7 +559,7 @@ void FSteamVRHMD::Startup()
 		float IdealScreenPercentage = FMath::Max(WidthPercentage, HeightPercentage);
 
 		//@todo steamvr: move out of here
-		IConsoleVariable* CScrPercVar = IConsoleManager::Get().FindConsoleVariable(TEXT("r.ScreenPercentage"));
+		static IConsoleVariable* CScrPercVar = IConsoleManager::Get().FindConsoleVariable(TEXT("r.ScreenPercentage"));
 
 		if (FMath::RoundToInt(CScrPercVar->GetFloat()) != FMath::RoundToInt(IdealScreenPercentage))
 		{
