[INTEGRATE] Change 2420757 by [REDACTED_PERSON][REDACTED_HANDLE] on 2015/01/27 13:47:43

	Fixed 'early access' SteamVR plugin not being marked beta

diff --git a/Engine/Plugins/Experimental/SteamVR/SteamVR.uplugin b/Engine/Plugins/Experimental/SteamVR/SteamVR.uplugin
--- a/Engine/Plugins/Experimental/SteamVR/SteamVR.uplugin
+++ b/Engine/Plugins/Experimental/SteamVR/SteamVR.uplugin
@@ -10,6 +10,7 @@
 	"Description" : "[EARLY ACCESS PREVIEW] Support for the SteamVR API",
 	"Category" : "Head Mounted Displays.Devices",
 	"EnabledByDefault" : false,
+	"IsBetaVersion" : true,
 
 	"Modules" :
 	[
