- Moved SteamVR plugin out of Experimental and into Runtime

diff --git a/Engine/Plugins/Experimental/SteamVR/Source/SteamVR/Classes/SteamVRChaperoneComponent.h b/Engine/Plugins/Experimental/SteamVR/Source/SteamVR/Classes/SteamVRChaperoneComponent.h
deleted file mode 100644
diff --git a/Engine/Plugins/Experimental/SteamVR/Source/SteamVR/Classes/SteamVRFunctionLibrary.h b/Engine/Plugins/Experimental/SteamVR/Source/SteamVR/Classes/SteamVRFunctionLibrary.h
deleted file mode 100644
diff --git a/Engine/Plugins/Experimental/SteamVR/Source/SteamVR/Private/SteamVRChaperoneComponent.cpp b/Engine/Plugins/Experimental/SteamVR/Source/SteamVR/Private/SteamVRChaperoneComponent.cpp
deleted file mode 100644
diff --git a/Engine/Plugins/Experimental/SteamVR/Source/SteamVR/Private/SteamVRFunctionLibrary.cpp b/Engine/Plugins/Experimental/SteamVR/Source/SteamVR/Private/SteamVRFunctionLibrary.cpp
deleted file mode 100644
diff --git a/Engine/Plugins/Experimental/SteamVR/Source/SteamVR/Private/SteamVRHMD.cpp b/Engine/Plugins/Experimental/SteamVR/Source/SteamVR/Private/SteamVRHMD.cpp
deleted file mode 100644
diff --git a/Engine/Plugins/Experimental/SteamVR/Source/SteamVR/Private/SteamVRHMD.h b/Engine/Plugins/Experimental/SteamVR/Source/SteamVR/Private/SteamVRHMD.h
deleted file mode 100644
diff --git a/Engine/Plugins/Experimental/SteamVR/Source/SteamVR/Private/SteamVRPrivatePCH.h b/Engine/Plugins/Experimental/SteamVR/Source/SteamVR/Private/SteamVRPrivatePCH.h
deleted file mode 100644
diff --git a/Engine/Plugins/Experimental/SteamVR/Source/SteamVR/Private/SteamVRRender.cpp b/Engine/Plugins/Experimental/SteamVR/Source/SteamVR/Private/SteamVRRender.cpp
deleted file mode 100644
diff --git a/Engine/Plugins/Experimental/SteamVR/Source/SteamVR/Public/ISteamVRPlugin.h b/Engine/Plugins/Experimental/SteamVR/Source/SteamVR/Public/ISteamVRPlugin.h
deleted file mode 100644
diff --git a/Engine/Plugins/Experimental/SteamVR/Source/SteamVR/SteamVR.Build.cs b/Engine/Plugins/Experimental/SteamVR/Source/SteamVR/SteamVR.Build.cs
deleted file mode 100644
diff --git a/Engine/Plugins/Experimental/SteamVR/SteamVR.uplugin b/Engine/Plugins/Experimental/SteamVR/SteamVR.uplugin
deleted file mode 100644
diff --git a/Engine/Plugins/Runtime/Steam/SteamVR/Source/SteamVR/Classes/SteamVRChaperoneComponent.h b/Engine/Plugins/Runtime/Steam/SteamVR/Source/SteamVR/Classes/SteamVRChaperoneComponent.h
new file mode 100644
--- /dev/null
+++ b/Engine/Plugins/Runtime/Steam/SteamVR/Source/SteamVR/Classes/SteamVRChaperoneComponent.h
@@ -0,0 +1,39 @@
+// Copyright 1998-2015 Epic Games, Inc. All Rights Reserved.
+
+#pragma once
+#include "SteamVRChaperoneComponent.generated.h"
+
+/**
+ * SteamVR Extensions Function Library
+ */
+UCLASS(MinimalAPI, meta = (BlueprintSpawnableComponent), ClassGroup = SteamVR)
+class USteamVRChaperoneComponent : public UActorComponent
+{
+	GENERATED_UCLASS_BODY()
+
+	DECLARE_DYNAMIC_MULTICAST_DELEGATE(FSteamVRChaperoneEvent);
+
+	UPROPERTY(BlueprintAssignable, Category = "SteamVR")
+	FSteamVRChaperoneEvent OnLeaveSoftBounds;
+
+	UPROPERTY(BlueprintAssignable, Category = "SteamVR")
+	FSteamVRChaperoneEvent OnReturnToSoftBounds;
+
+	void TickComponent(float DeltaTime, enum ELevelTick TickType, FActorComponentTickFunction *ThisTickFunction) override;
+
+	/**
+	 * Returns the soft bounds from the Chaperone, in Unreal-scale HMD-space coordinates, centered around the HMD's calibration origin (0,0,0).  Note that the Z component is always zero
+	 */
+	UFUNCTION(BlueprintPure, Category = "SteamVR")
+	TArray<FVector> GetSoftBounds() const;
+
+	/**
+	* Returns the hard bounds from the Chaperone, in Unreal-scale HMD-space coordinates, centered around the HMD's calibration origin (0,0,0).  Each set of four bounds will form a quad to define a set of hard bounds
+	*/
+	UFUNCTION(BlueprintPure, Category = "SteamVR")
+	TArray<FVector> GetHardBounds() const;
+
+private:
+	/** Whether or not we were inside the soft bounds last time, so we can detect changes */
+	bool bWasInsideSoftBounds;
+};
diff --git a/Engine/Plugins/Runtime/Steam/SteamVR/Source/SteamVR/Classes/SteamVRFunctionLibrary.h b/Engine/Plugins/Runtime/Steam/SteamVR/Source/SteamVR/Classes/SteamVRFunctionLibrary.h
new file mode 100644
--- /dev/null
+++ b/Engine/Plugins/Runtime/Steam/SteamVR/Source/SteamVR/Classes/SteamVRFunctionLibrary.h
@@ -0,0 +1,42 @@
+// Copyright 1998-2015 Epic Games, Inc. All Rights Reserved.
+
+#pragma once
+#include "SteamVRFunctionLibrary.generated.h"
+
+/**
+ * SteamVR Extensions Function Library
+ */
+UCLASS(MinimalAPI)
+class USteamVRFunctionLibrary : public UBlueprintFunctionLibrary
+{
+	GENERATED_UCLASS_BODY()
+	
+	/**
+	 * Returns an array of the currently tracked device IDs
+	 *
+	 * @param	OutTrackedDeviceIds	(out) Array containing the ID of each device that's currently tracked
+	 */
+	UFUNCTION(BlueprintPure, Category="SteamVR")
+	static void GetValidTrackedDeviceIds(TArray<int32>& OutTrackedDeviceIds);
+
+	/**
+	 * Gets the orientation and position (in device space) of the device with the specified ID
+	 *
+	 * @param	DeviceId		Id of the device to get tracking info for
+	 * @param	OutPosition		(out) Current position of the device
+	 * @param	OutOrientation	(out) Current orientation of the device
+	 * @return	True if the specified device id had a valid tracking pose this frame, false otherwise
+	 */
+	UFUNCTION(BlueprintPure, Category = "SteamVR")
+	static bool GetTrackedDevicePositionAndOrientation(int32 DeviceId, FVector& OutPosition, FRotator& OutOrientation);
+
+	/**
+	 * Given a controller index, returns the attached tracked device ID
+	 *
+	 * @param	ControllerIndex	Index of the controller to get the tracked device ID for
+	 * @param	OutDeviceId		(out) Tracked device ID for the controller
+	 * @return	True if the specified controller index has a valid tracked device ID
+	 */
+	UFUNCTION(BlueprintPure, Category = "SteamVR")
+	static bool GetTrackedDeviceIdFromControllerIndex(int32 ControllerIndex, int32& OutDeviceId);
+};
diff --git a/Engine/Plugins/Runtime/Steam/SteamVR/Source/SteamVR/Private/SteamVRChaperoneComponent.cpp b/Engine/Plugins/Runtime/Steam/SteamVR/Source/SteamVR/Private/SteamVRChaperoneComponent.cpp
new file mode 100644
--- /dev/null
+++ b/Engine/Plugins/Runtime/Steam/SteamVR/Source/SteamVR/Private/SteamVRChaperoneComponent.cpp
@@ -0,0 +1,70 @@
+// Copyright 1998-2015 Epic Games, Inc. All Rights Reserved.
+//
+#include "SteamVRPrivatePCH.h"
+#include "SteamVRHMD.h"
+#include "Classes/SteamVRChaperoneComponent.h"
+
+USteamVRChaperoneComponent::USteamVRChaperoneComponent(const FObjectInitializer& ObjectInitializer)
+: Super(ObjectInitializer)
+{
+	PrimaryComponentTick.bCanEverTick = true;
+	PrimaryComponentTick.bStartWithTickEnabled = true;
+	PrimaryComponentTick.TickGroup = TG_DuringPhysics;
+
+	bTickInEditor = true;
+	bAutoActivate = true;
+
+	bWasInsideSoftBounds = true;
+}
+
+void USteamVRChaperoneComponent::TickComponent(float DeltaTime, enum ELevelTick TickType, FActorComponentTickFunction *ThisTickFunction)
+{
+	Super::TickComponent(DeltaTime, TickType, ThisTickFunction);
+
+	//@todo Make this safe once we can add something to the DeviceType enum.  For now, make the terrible assumption this is a SteamVR device.
+ 	FSteamVRHMD* SteamVRHMD = (FSteamVRHMD*)(GEngine->HMDDevice.Get());
+ 	if (SteamVRHMD && SteamVRHMD->IsStereoEnabled())
+ 	{
+ 		bool bInSoftBounds = SteamVRHMD->IsInsideSoftBounds();
+ 
+		if (bInSoftBounds != bWasInsideSoftBounds)
+		{
+			if (bInSoftBounds)
+			{
+				OnReturnToSoftBounds.Broadcast();
+			}
+			else
+			{
+				OnLeaveSoftBounds.Broadcast();
+			}
+		}
+
+		bWasInsideSoftBounds = bInSoftBounds;
+ 	}
+}
+
+TArray<FVector> USteamVRChaperoneComponent::GetSoftBounds() const
+{
+	TArray<FVector> RetValue;
+
+	FSteamVRHMD* SteamVRHMD = (FSteamVRHMD*)(GEngine->HMDDevice.Get());
+	if (SteamVRHMD && SteamVRHMD->IsStereoEnabled())
+	{
+		RetValue = SteamVRHMD->GetSoftBounds();
+	}
+
+	return RetValue;
+}
+
+TArray<FVector> USteamVRChaperoneComponent::GetHardBounds() const
+{
+	TArray<FVector> RetValue;
+
+	FSteamVRHMD* SteamVRHMD = (FSteamVRHMD*)(GEngine->HMDDevice.Get());
+	if (SteamVRHMD && SteamVRHMD->IsStereoEnabled())
+	{
+		RetValue = SteamVRHMD->GetHardBounds();
+	}
+
+	return RetValue;
+}
\ No newline at end of file
diff --git a/Engine/Plugins/Runtime/Steam/SteamVR/Source/SteamVR/Private/SteamVRFunctionLibrary.cpp b/Engine/Plugins/Runtime/Steam/SteamVR/Source/SteamVR/Private/SteamVRFunctionLibrary.cpp
new file mode 100644
--- /dev/null
+++ b/Engine/Plugins/Runtime/Steam/SteamVR/Source/SteamVR/Private/SteamVRFunctionLibrary.cpp
@@ -0,0 +1,53 @@
+// Copyright 1998-2015 Epic Games, Inc. All Rights Reserved.
+//
+#include "SteamVRPrivatePCH.h"
+#include "SteamVRHMD.h"
+#include "Classes/SteamVRFunctionLibrary.h"
+
+USteamVRFunctionLibrary::USteamVRFunctionLibrary(const FObjectInitializer& ObjectInitializer)
+: Super(ObjectInitializer)
+{
+}
+
+//@TODO steamvr:  Make these safer with device ID checking instead of nasty casting.
+
+void USteamVRFunctionLibrary::GetValidTrackedDeviceIds(TArray<int32>& DeviceIds)
+{
+	DeviceIds.Empty();
+
+	if (GEngine->HMDDevice.IsValid())
+	{
+		FSteamVRHMD* SteamVRHMD = static_cast<FSteamVRHMD*>(GEngine->HMDDevice.Get());
+		SteamVRHMD->GetTrackedDeviceIds(DeviceIds);
+	}
+}
+
+bool USteamVRFunctionLibrary::GetTrackedDevicePositionAndOrientation(int32 DeviceId, FVector& OutPosition, FRotator& OutOrientation)
+{
+	bool RetVal = false;
+
+	if (GEngine->HMDDevice.IsValid())
+	{
+		FSteamVRHMD* SteamVRHMD = static_cast<FSteamVRHMD*>(GEngine->HMDDevice.Get());
+
+		FQuat DeviceOrientation = FQuat::Identity;
+		RetVal = SteamVRHMD->GetTrackedObjectOrientationAndPosition(DeviceId, DeviceOrientation, OutPosition);
+		OutOrientation = DeviceOrientation.Rotator();
+	}
+
+	return RetVal;
+}
+
+bool USteamVRFunctionLibrary::GetTrackedDeviceIdFromControllerIndex(int32 ControllerIndex, int32& OutDeviceId)
+{
+	OutDeviceId = -1;
+
+	if (!GEngine->HMDDevice.IsValid())
+	{
+		return false;
+	}
+
+	FSteamVRHMD* SteamVRHMD = static_cast<FSteamVRHMD*>(GEngine->HMDDevice.Get());
+
+	return SteamVRHMD->GetTrackedDeviceIdFromControllerIndex(ControllerIndex, OutDeviceId);
+}
diff --git a/Engine/Plugins/Runtime/Steam/SteamVR/Source/SteamVR/Private/SteamVRHMD.cpp b/Engine/Plugins/Runtime/Steam/SteamVR/Source/SteamVR/Private/SteamVRHMD.cpp
new file mode 100644
--- /dev/null
+++ b/Engine/Plugins/Runtime/Steam/SteamVR/Source/SteamVR/Private/SteamVRHMD.cpp
@@ -0,0 +1,1052 @@
+// Copyright 1998-2015 Epic Games, Inc. All Rights Reserved.
+
+#include "SteamVRPrivatePCH.h"
+#include "SteamVRHMD.h"
+
+#include "RendererPrivate.h"
+#include "ScenePrivate.h"
+#include "PostProcess/PostProcessHMD.h"
+#include "Classes/SteamVRFunctionLibrary.h"
+
+
+//---------------------------------------------------
+// SteamVR Plugin Implementation
+//---------------------------------------------------
+
+class FSteamVRPlugin : public ISteamVRPlugin
+{
+	/** IHeadMountedDisplayModule implementation */
+	virtual TSharedPtr< class IHeadMountedDisplay > CreateHeadMountedDisplay() override;
+
+	FString GetModulePriorityKeyName() const
+	{
+		return FString(TEXT("SteamVR"));
+	}
+};
+
+IMPLEMENT_MODULE( FSteamVRPlugin, SteamVR )
+
+TSharedPtr< class IHeadMountedDisplay > FSteamVRPlugin::CreateHeadMountedDisplay()
+{
+#if STEAMVR_SUPPORTED_PLATFORMS
+	TSharedPtr< FSteamVRHMD > SteamVRHMD( new FSteamVRHMD() );
+	if( SteamVRHMD->IsInitialized() )
+	{
+		return SteamVRHMD;
+	}
+#endif//STEAMVR_SUPPORTED_PLATFORMS
+	return nullptr;
+}
+
+
+//---------------------------------------------------
+// SteamVR IHeadMountedDisplay Implementation
+//---------------------------------------------------
+
+#if STEAMVR_SUPPORTED_PLATFORMS
+
+
+bool FSteamVRHMD::IsHMDEnabled() const
+{
+	return bHmdEnabled;
+}
+
+void FSteamVRHMD::EnableHMD(bool enable)
+{
+	bHmdEnabled = enable;
+
+	if (!bHmdEnabled)
+	{
+		EnableStereo(false);
+	}
+}
+
+EHMDDeviceType::Type FSteamVRHMD::GetHMDDeviceType() const
+{
+	return EHMDDeviceType::DT_NoPost;
+}
+
+bool FSteamVRHMD::GetHMDMonitorInfo(MonitorInfo& MonitorDesc) 
+{
+	if (IsInitialized())
+	{
+		int32 X, Y;
+		uint32 Width, Height;
+		Hmd->GetWindowBounds(&X, &Y, &Width, &Height);
+
+		MonitorDesc.MonitorName = DisplayId;
+		MonitorDesc.MonitorId	= 0;
+		MonitorDesc.DesktopX	= X;
+		MonitorDesc.DesktopY	= Y;
+		MonitorDesc.ResolutionX = Width;
+		MonitorDesc.ResolutionY = Height;
+
+		return true;
+	}
+	else
+	{
+		MonitorDesc.MonitorName = "";
+		MonitorDesc.MonitorId = 0;
+		MonitorDesc.DesktopX = MonitorDesc.DesktopY = MonitorDesc.ResolutionX = MonitorDesc.ResolutionY = 0;
+	}
+
+	return false;
+}
+
+void FSteamVRHMD::GetFieldOfView(float& OutHFOVInDegrees, float& OutVFOVInDegrees) const
+{
+	OutHFOVInDegrees = 0.0f;
+	OutVFOVInDegrees = 0.0f;
+}
+
+bool FSteamVRHMD::DoesSupportPositionalTracking() const
+{
+	return true;
+}
+
+bool FSteamVRHMD::HasValidTrackingPosition() const
+{
+	return bHmdPosTracking && bHaveVisionTracking;
+}
+
+void FSteamVRHMD::GetPositionalTrackingCameraProperties(FVector& OutOrigin, FRotator& OutOrientation, float& OutHFOV, float& OutVFOV, float& OutCameraDistance, float& OutNearPlane, float& OutFarPlane) const
+{
+}
+
+void FSteamVRHMD::SetInterpupillaryDistance(float NewInterpupillaryDistance)
+{
+}
+
+float FSteamVRHMD::GetInterpupillaryDistance() const
+{
+	return 0.064f;
+}
+
+void FSteamVRHMD::GetCurrentPose(FQuat& CurrentOrientation, FVector& CurrentPosition, uint32 DeviceId, bool bForceRefresh /* = false*/)
+{
+	if (!Hmd)
+	{
+		return;
+	}
+
+	check(DeviceId < vr::k_unMaxTrackedDeviceCount);
+
+	if (bForceRefresh)
+	{
+		// With SteamVR, we should only update on the PreRender_ViewFamily, and then the next frame should use the previous frame's results
+		check(IsInRenderingThread());
+
+		TrackingFrame.FrameNumber = GFrameNumberRenderThread;
+
+		vr::TrackedDevicePose_t Poses[vr::k_unMaxTrackedDeviceCount];
+		VRCompositor->WaitGetPoses(Poses, ARRAYSIZE(Poses));
+
+		for (uint32 i = 0; i < vr::k_unMaxTrackedDeviceCount; ++i)
+		{
+			TrackingFrame.bDeviceIsValid[i] = Poses[i].bDeviceIsValid;
+			TrackingFrame.bPoseIsValid[i] = Poses[i].bPoseIsValid;
+
+			FVector CurrentPosition;
+			FQuat CurrentOrientation;
+			PoseToOrientationAndPosition(Poses[i].mDeviceToAbsoluteTracking, CurrentOrientation, CurrentPosition);
+
+			TrackingFrame.DeviceOrientation[i] = CurrentOrientation;
+			TrackingFrame.DevicePosition[i] = CurrentPosition;
+
+			TrackingFrame.RawPoses[i] = Poses[i].mDeviceToAbsoluteTracking;
+		}
+	}
+	
+	// Update CurrentOrientation and CurrentPosition for the desired device, if valid
+	if (TrackingFrame.bPoseIsValid[DeviceId])
+ 	{
+		CurrentOrientation = TrackingFrame.DeviceOrientation[DeviceId];
+		CurrentPosition = TrackingFrame.DevicePosition[DeviceId];
+ 	}
+}
+
+
+
+bool FSteamVRHMD::IsInsideSoftBounds()
+{
+	if (VRChaperone)
+	{
+		vr::HmdMatrix34_t VRPose = TrackingFrame.RawPoses[vr::k_unTrackedDeviceIndex_Hmd];
+		FMatrix Pose = ToFMatrix(VRPose);
+		
+		const FVector HMDLocation(Pose.M[3][0], 0.f, Pose.M[3][2]);
+
+		bool bLastWasNegative = false;
+
+		// Since the order of the soft bounds are points on a plane going clockwise, wind around the sides, checking the crossproduct of the affine side to the affine HMD position.  If they're all on the same side, we're in the bounds
+		for (uint8 i = 0; i < 4; ++i)
+		{
+			const FVector PointA = ChaperoneBounds.SoftBounds.Corners[i];
+			const FVector PointB = ChaperoneBounds.SoftBounds.Corners[(i + 1) % 4];
+
+			const FVector AffineSegment = PointB - PointA;
+			const FVector AffinePoint = HMDLocation - PointA;
+			const FVector CrossProduct = FVector::CrossProduct(AffineSegment, AffinePoint);
+
+			const bool bIsNegative = (CrossProduct.Y < 0);
+
+			// If the cross between the point and the side has flipped, that means we're not consistent, and therefore outside the bounds
+			if ((i > 0) && (bLastWasNegative != bIsNegative))
+			{
+				return false;
+			}
+
+			bLastWasNegative = bIsNegative;
+		}
+
+		return true;
+	}
+
+	return false;
+}
+
+/** Helper function to convert bounds from SteamVR space to scaled Unreal space*/
+TArray<FVector> ConvertBoundsToUnrealSpace(const FBoundingQuad& InBounds, const float WorldToMetersScale)
+{
+	TArray<FVector> Bounds;
+
+	for (int32 i = 0; i < ARRAYSIZE(InBounds.Corners); ++i)
+	{
+		const FVector SteamVRCorner = InBounds.Corners[i];
+		const FVector UnrealVRCorner(-SteamVRCorner.Z, SteamVRCorner.X, SteamVRCorner.Y);
+		Bounds.Add(UnrealVRCorner * WorldToMetersScale);
+	}
+
+	return Bounds;
+}
+
+TArray<FVector> FSteamVRHMD::GetSoftBounds() const
+{
+	return ConvertBoundsToUnrealSpace(ChaperoneBounds.SoftBounds, WorldToMetersScale);
+}
+
+TArray<FVector> FSteamVRHMD::GetHardBounds() const
+{
+	TArray<FVector> Bounds;
+
+	for (int32 i = 0; i < ChaperoneBounds.HardBounds.Num(); ++i)
+	{
+		Bounds.Append(ConvertBoundsToUnrealSpace(ChaperoneBounds.HardBounds[i], WorldToMetersScale));
+	}
+	
+	return Bounds;
+}
+
+void FSteamVRHMD::PoseToOrientationAndPosition(const vr::HmdMatrix34_t& InPose, FQuat& OutOrientation, FVector& OutPosition) const
+{
+	FMatrix Pose = ToFMatrix(InPose);
+	FQuat Orientation(Pose);
+ 
+	OutOrientation.X = -Orientation.Z;
+	OutOrientation.Y = Orientation.X;
+	OutOrientation.Z = Orientation.Y;
+ 	OutOrientation.W = -Orientation.W;	
+
+	FVector Position = FVector(-Pose.M[3][2], Pose.M[3][0], Pose.M[3][1]) * WorldToMetersScale;
+	OutPosition = BaseOrientation.Inverse().RotateVector(Position);
+
+	OutOrientation = BaseOrientation.Inverse() * OutOrientation;
+	OutOrientation.Normalize();
+}
+
+void FSteamVRHMD::GetCurrentOrientationAndPosition(FQuat& CurrentOrientation, FVector& CurrentPosition)
+{
+	check(IsInGameThread());
+	GetCurrentPose(CurHmdOrientation, CurHmdPosition);
+	CurrentOrientation = LastHmdOrientation = CurHmdOrientation;
+
+	CurrentPosition = CurHmdPosition;
+}
+
+void FSteamVRHMD::GetTrackedDeviceIds(TArray<int32>& TrackedIds)
+{
+	TrackedIds.Empty();
+
+	for (uint32 i = 0; i < vr::k_unMaxTrackedDeviceCount; ++i)
+	{
+		// Add only devices with a currently valid tracked pose, and exclude the HMD
+		if ((i != vr::k_unTrackedDeviceIndex_Hmd) && TrackingFrame.bPoseIsValid[i])
+	{
+			TrackedIds.Add(i);
+		}
+	}
+}
+
+bool FSteamVRHMD::GetTrackedObjectOrientationAndPosition(uint32 DeviceId, FQuat& CurrentOrientation, FVector& CurrentPosition)
+{
+	check(IsInGameThread());
+
+	bool bHasValidPose = false;
+
+	if (DeviceId < vr::k_unMaxTrackedDeviceCount)
+	{
+		CurrentOrientation = TrackingFrame.DeviceOrientation[DeviceId];
+		CurrentPosition = TrackingFrame.DevicePosition[DeviceId];
+
+		bHasValidPose = TrackingFrame.bPoseIsValid[DeviceId] && TrackingFrame.bDeviceIsValid[DeviceId];
+	}
+
+	return bHasValidPose;
+}
+
+bool FSteamVRHMD::GetTrackedDeviceIdFromControllerIndex(int32 ControllerIndex, int32& OutDeviceId)
+{
+	check(IsInGameThread());
+
+	OutDeviceId = -1;
+
+	if ((ControllerIndex < 0) || (ControllerIndex >= MAX_STEAM_CONTROLLERS))
+	{
+		return false;
+	}
+
+	OutDeviceId = ControllerIndexToTrackedDeviceIndex[ControllerIndex];
+	
+	return (OutDeviceId != -1);
+}
+
+ISceneViewExtension* FSteamVRHMD::GetViewExtension()
+{
+	return this;
+}
+
+void FSteamVRHMD::ApplyHmdRotation(APlayerController* PC, FRotator& ViewRotation)
+{
+	ViewRotation.Normalize();
+
+	GetCurrentPose(CurHmdOrientation, CurHmdPosition);
+	LastHmdOrientation = CurHmdOrientation;
+
+	const FRotator DeltaRot = ViewRotation - PC->GetControlRotation();
+	DeltaControlRotation = (DeltaControlRotation + DeltaRot).GetNormalized();
+
+	// Pitch from other sources is never good, because there is an absolute up and down that must be respected to avoid motion sickness.
+	// Same with roll.
+	DeltaControlRotation.Pitch = 0;
+	DeltaControlRotation.Roll = 0;
+	DeltaControlOrientation = DeltaControlRotation.Quaternion();
+
+	ViewRotation = FRotator(DeltaControlOrientation * CurHmdOrientation);
+}
+
+void FSteamVRHMD::UpdatePlayerCameraRotation(APlayerCameraManager* Camera, struct FMinimalViewInfo& POV)
+{
+	GetCurrentPose(CurHmdOrientation, CurHmdPosition);
+	LastHmdOrientation = CurHmdOrientation;
+
+	DeltaControlRotation = POV.Rotation;
+	DeltaControlOrientation = DeltaControlRotation.Quaternion();
+
+	// Apply HMD orientation to camera rotation.
+	POV.Rotation = FRotator(POV.Rotation.Quaternion() * CurHmdOrientation);
+}
+
+bool FSteamVRHMD::IsChromaAbCorrectionEnabled() const
+{
+	return true;
+}
+
+bool FSteamVRHMD::Exec( UWorld* InWorld, const TCHAR* Cmd, FOutputDevice& Ar )
+{
+	if (FParse::Command( &Cmd, TEXT("STEREO") ))
+	{
+		if (FParse::Command(&Cmd, TEXT("ON")))
+		{
+			if (!IsHMDEnabled())
+			{
+				Ar.Logf(TEXT("HMD is disabled. Use 'hmd enable' to re-enable it."));
+			}
+			EnableStereo(true);
+			return true;
+		}
+		else if (FParse::Command(&Cmd, TEXT("OFF")))
+		{
+			EnableStereo(false);
+			return true;
+		}
+
+		float val;
+		if (FParse::Value(Cmd, TEXT("E="), val))
+		{
+			IPD = val;
+		}
+	}
+	else if (FParse::Command(&Cmd, TEXT("HMD")))
+	{
+		if (FParse::Command(&Cmd, TEXT("ENABLE")))
+		{
+			EnableHMD(true);
+			return true;
+		}
+		else if (FParse::Command(&Cmd, TEXT("DISABLE")))
+		{
+			EnableHMD(false);
+			return true;
+		}
+
+		int32 val;
+		if (FParse::Value(Cmd, TEXT("MIRROR"), val))
+		{
+			if ((val >= 0) && (val <= 2))
+			{
+				WindowMirrorMode = val;
+			}
+			else
+		{
+				Ar.Logf(TEXT("HMD MIRROR accepts values from 0 though 2"));
+		}
+
+			return true;
+		}
+	}
+	else if (FParse::Command(&Cmd, TEXT("UNCAPFPS")))
+	{
+		GEngine->bSmoothFrameRate = false;
+		return true;
+	}
+
+	return false;
+}
+
+void FSteamVRHMD::OnScreenModeChange(EWindowMode::Type WindowMode)
+{
+	EnableStereo(WindowMode != EWindowMode::Windowed);
+}
+
+bool FSteamVRHMD::IsPositionalTrackingEnabled() const
+{
+	return bHmdPosTracking;
+}
+
+bool FSteamVRHMD::EnablePositionalTracking(bool enable)
+{
+	bHmdPosTracking = enable;
+	return IsPositionalTrackingEnabled();
+}
+
+bool FSteamVRHMD::IsHeadTrackingAllowed() const
+{
+	return GEngine->IsStereoscopic3D();
+}
+
+bool FSteamVRHMD::IsInLowPersistenceMode() const
+{
+	return true;
+}
+
+void FSteamVRHMD::EnableLowPersistenceMode(bool Enable)
+{
+}
+
+
+void FSteamVRHMD::ResetOrientationAndPosition(float yaw)
+{
+	ResetOrientation(yaw);
+	ResetPosition();
+}
+
+void FSteamVRHMD::ResetOrientation(float Yaw)
+{
+	FRotator ViewRotation;
+	ViewRotation = FRotator(TrackingFrame.DeviceOrientation[vr::k_unTrackedDeviceIndex_Hmd]);
+	ViewRotation.Pitch = 0;
+	ViewRotation.Roll = 0;
+
+	if (Yaw != 0.f)
+	{
+		// apply optional yaw offset
+		ViewRotation.Yaw -= Yaw;
+		ViewRotation.Normalize();
+	}
+
+	BaseOrientation = ViewRotation.Quaternion();
+}
+void FSteamVRHMD::ResetPosition()
+{
+	//@todo steamvr: ResetPosition()
+}
+
+void FSteamVRHMD::SetClippingPlanes(float NCP, float FCP)
+{
+}
+
+void FSteamVRHMD::SetBaseRotation(const FRotator& BaseRot)
+{
+}
+FRotator FSteamVRHMD::GetBaseRotation() const
+{
+	return FRotator::ZeroRotator;
+}
+
+void FSteamVRHMD::SetBaseOrientation(const FQuat& BaseOrient)
+{
+	BaseOrientation = BaseOrient;
+}
+
+FQuat FSteamVRHMD::GetBaseOrientation() const
+{
+	return BaseOrientation;
+}
+
+void FSteamVRHMD::SetPositionOffset(const FVector& PosOff)
+{
+}
+
+FVector FSteamVRHMD::GetPositionOffset() const
+{
+	return FVector::ZeroVector;
+}
+
+bool FSteamVRHMD::IsStereoEnabled() const
+{
+	return bStereoEnabled && bHmdEnabled;
+}
+
+bool FSteamVRHMD::EnableStereo(bool stereo)
+{
+	bStereoEnabled = (IsHMDEnabled()) ? stereo : false;
+
+	FSystemResolution::RequestResolutionChange(1280, 720, (stereo) ? EWindowMode::WindowedMirror : EWindowMode::Windowed);
+
+	return bStereoEnabled;
+}
+
+void FSteamVRHMD::AdjustViewRect(EStereoscopicPass StereoPass, int32& X, int32& Y, uint32& SizeX, uint32& SizeY) const
+{
+	//@todo steamvr: get the actual rects from steamvr
+	SizeX = SizeX / 2;
+	if( StereoPass == eSSP_RIGHT_EYE )
+	{
+		X += SizeX;
+	}
+}
+
+void FSteamVRHMD::CalculateStereoViewOffset(const enum EStereoscopicPass StereoPassType, const FRotator& ViewRotation, const float WorldToMeters, FVector& ViewLocation)
+{
+	if( StereoPassType != eSSP_FULL)
+	{
+		vr::Hmd_Eye HmdEye = (StereoPassType == eSSP_LEFT_EYE) ? vr::Eye_Left : vr::Eye_Right;
+		vr::HmdMatrix34_t HeadFromEye = Hmd->GetEyeToHeadTransform(HmdEye);
+
+		// grab the eye position, currently ignoring the rotation supplied by GetHeadFromEyePose()
+		FVector TotalOffset = FVector(-HeadFromEye.m[2][3], HeadFromEye.m[0][3], HeadFromEye.m[1][3]) * WorldToMeters;
+
+		ViewLocation += ViewRotation.Quaternion().RotateVector(TotalOffset);
+
+ 		const FVector vHMDPosition = DeltaControlOrientation.RotateVector(TrackingFrame.DevicePosition[vr::k_unTrackedDeviceIndex_Hmd]);
+		ViewLocation += vHMDPosition;
+	}
+}
+
+FMatrix FSteamVRHMD::GetStereoProjectionMatrix(const enum EStereoscopicPass StereoPassType, const float FOV) const
+{
+	check(IsStereoEnabled());
+
+	vr::Hmd_Eye HmdEye = (StereoPassType == eSSP_LEFT_EYE) ? vr::Eye_Left : vr::Eye_Right;
+	float Left, Right, Top, Bottom;
+
+	Hmd->GetProjectionRaw(HmdEye, &Right, &Left, &Top, &Bottom);
+	Bottom *= -1.0f;
+	Top *= -1.0f;
+	Right *= -1.0f;
+	Left *= -1.0f;
+
+	float ZNear = GNearClippingPlane;
+
+	float SumRL = (Right + Left);
+	float SumTB = (Top + Bottom);
+	float InvRL = (1.0f / (Right - Left));
+	float InvTB = (1.0f / (Top - Bottom));
+
+#if 1
+	FMatrix Mat = FMatrix(
+		FPlane((2.0f * InvRL), 0.0f, 0.0f, 0.0f),
+		FPlane(0.0f, (2.0f * InvTB), 0.0f, 0.0f),
+		FPlane((SumRL * InvRL), (SumTB * InvTB), 0.0f, 1.0f),
+		FPlane(0.0f, 0.0f, ZNear, 0.0f)
+		);
+#else
+	vr::HmdMatrix44_t SteamMat = Hmd->GetProjectionMatrix(HmdEye, ZNear, 10000.0f, vr::API_DirectX);
+	FMatrix Mat = ToFMatrix(SteamMat);
+
+	Mat.M[3][3] = 0.0f;
+	Mat.M[2][3] = 1.0f;
+	Mat.M[2][2] = 0.0f;
+	Mat.M[3][2] = ZNear;
+#endif
+
+	return Mat;
+
+}
+
+void FSteamVRHMD::InitCanvasFromView(FSceneView* InView, UCanvas* Canvas)
+{
+}
+
+void FSteamVRHMD::PushViewportCanvas(EStereoscopicPass StereoPass, FCanvas *InCanvas, UCanvas *InCanvasObject, FViewport *InViewport) const 
+{
+	FMatrix m;
+	m.SetIdentity();
+	InCanvas->PushAbsoluteTransform(m);
+}
+
+void FSteamVRHMD::PushViewCanvas(EStereoscopicPass StereoPass, FCanvas *InCanvas, UCanvas *InCanvasObject, FSceneView *InView) const 
+{
+	FMatrix m;
+	m.SetIdentity();
+	InCanvas->PushAbsoluteTransform(m);
+}
+
+void FSteamVRHMD::GetEyeRenderParams_RenderThread(EStereoscopicPass StereoPass, FVector2D& EyeToSrcUVScaleValue, FVector2D& EyeToSrcUVOffsetValue) const
+{
+	if (StereoPass == eSSP_LEFT_EYE)
+	{
+		EyeToSrcUVOffsetValue.X = 0.0f;
+		EyeToSrcUVOffsetValue.Y = 0.0f;
+
+		EyeToSrcUVScaleValue.X = 0.5f;
+		EyeToSrcUVScaleValue.Y = 1.0f;
+	}
+	else
+	{
+		EyeToSrcUVOffsetValue.X = 0.5f;
+		EyeToSrcUVOffsetValue.Y = 0.0f;
+
+		EyeToSrcUVScaleValue.X = 0.5f;
+		EyeToSrcUVScaleValue.Y = 1.0f;
+	}
+}
+
+
+void FSteamVRHMD::ModifyShowFlags(FEngineShowFlags& ShowFlags)
+{
+	ShowFlags.MotionBlur = 0;
+	ShowFlags.HMDDistortion = false;
+	ShowFlags.StereoRendering = IsStereoEnabled();
+}
+
+void FSteamVRHMD::SetupView(FSceneViewFamily& InViewFamily, FSceneView& InView)
+{
+	InView.BaseHmdOrientation = LastHmdOrientation;
+	InView.BaseHmdLocation = LastHmdPosition;
+	WorldToMetersScale = InView.WorldToMetersScale;
+	InViewFamily.bUseSeparateRenderTarget = true;
+}
+
+void FSteamVRHMD::PreRenderView_RenderThread(FSceneView& View)
+{
+	check(IsInRenderingThread());
+
+	// The last view location used to set the view will be in BaseHmdOrientation.  We need to calculate the delta from that, so that
+	// cameras that rely on game objects (e.g. other components) for their positions don't need to be updated on the render thread.
+	const FQuat DeltaOrient = View.BaseHmdOrientation.Inverse() * TrackingFrame.DeviceOrientation[vr::k_unTrackedDeviceIndex_Hmd];
+	View.ViewRotation = FRotator(View.ViewRotation.Quaternion() * DeltaOrient);
+ 	View.UpdateViewMatrix();
+}
+
+void FSteamVRHMD::PreRenderViewFamily_RenderThread(FSceneViewFamily& ViewFamily)
+{
+	check(IsInRenderingThread());
+	GetActiveRHIBridgeImpl()->BeginRendering();
+
+	FVector CurrentPosition;
+	FQuat CurrentOrientation;
+	GetCurrentPose(CurrentOrientation, CurrentPosition, vr::k_unTrackedDeviceIndex_Hmd, true);
+}
+
+void FSteamVRHMD::UpdateViewport(bool bUseSeparateRenderTarget, const FViewport& InViewport, SViewport* ViewportWidget)
+{
+	check(IsInGameThread());
+
+	FRHIViewport* const ViewportRHI = InViewport.GetViewportRHI().GetReference();
+
+	if (!IsStereoEnabled())
+	{
+		if (!bUseSeparateRenderTarget)
+		{
+			ViewportRHI->SetCustomPresent(nullptr);
+		}
+		return;
+	}
+
+	GetActiveRHIBridgeImpl()->UpdateViewport(InViewport, ViewportRHI);
+}
+
+FSteamVRHMD::BridgeBaseImpl* FSteamVRHMD::GetActiveRHIBridgeImpl()
+{
+#if PLATFORM_WINDOWS
+	if (pD3D11Bridge)
+	{
+		return pD3D11Bridge;
+	}
+#endif
+
+	return nullptr;
+}
+
+void FSteamVRHMD::CalculateRenderTargetSize(uint32& InOutSizeX, uint32& InOutSizeY) const
+{
+	check(IsInGameThread());
+
+//	if (Flags.bScreenPercentageEnabled)
+	{
+		static const auto CVar = IConsoleManager::Get().FindTConsoleVariableDataFloat(TEXT("r.ScreenPercentage"));
+		float value = CVar->GetValueOnGameThread();
+		if (value > 0.0f)
+		{
+			InOutSizeX = FMath::CeilToInt(InOutSizeX * value / 100.f);
+			InOutSizeY = FMath::CeilToInt(InOutSizeY * value / 100.f);
+		}
+	}
+}
+
+bool FSteamVRHMD::NeedReAllocateViewportRenderTarget(const FViewport& Viewport) const
+{
+	check(IsInGameThread());
+
+	if (IsStereoEnabled())
+	{
+		const uint32 InSizeX = Viewport.GetSizeXY().X;
+		const uint32 InSizeY = Viewport.GetSizeXY().Y;
+		FIntPoint RenderTargetSize;
+		RenderTargetSize.X = Viewport.GetRenderTargetTexture()->GetSizeX();
+		RenderTargetSize.Y = Viewport.GetRenderTargetTexture()->GetSizeY();
+
+		uint32 NewSizeX = InSizeX, NewSizeY = InSizeY;
+		CalculateRenderTargetSize(NewSizeX, NewSizeY);
+		if (NewSizeX != RenderTargetSize.X || NewSizeY != RenderTargetSize.Y)
+		{
+			return true;
+		}
+	}
+	return false;
+}
+
+FSteamVRHMD::FSteamVRHMD() :
+	Hmd(nullptr),
+	bHmdEnabled(true),
+	bStereoEnabled(false),
+	bHmdPosTracking(true),
+	bHaveVisionTracking(false),
+	IPD(0.064f),
+	WindowMirrorMode(0),
+	CurHmdOrientation(FQuat::Identity),
+	LastHmdOrientation(FQuat::Identity),
+	BaseOrientation(FQuat::Identity),
+	DeltaControlRotation(FRotator::ZeroRotator),
+	DeltaControlOrientation(FQuat::Identity),
+	CurHmdPosition(FVector::ZeroVector),
+	WorldToMetersScale(100.0f),
+	RendererModule(nullptr),
+	SteamDLLHandle(nullptr),
+	IdealScreenPercentage(100.0f)
+{
+	Startup();
+}
+
+FSteamVRHMD::~FSteamVRHMD()
+{
+	Shutdown();
+}
+
+bool FSteamVRHMD::IsInitialized() const
+{
+	return (Hmd != nullptr);
+}
+
+void FSteamVRHMD::Startup()
+{
+	// load the steam library
+	if (!LoadSteamModule())
+	{
+		return;
+	}
+
+	// grab a pointer to the renderer module for displaying our mirror window
+	static const FName RendererModuleName("Renderer");
+	RendererModule = FModuleManager::GetModulePtr<IRendererModule>(RendererModuleName);
+	
+	bool ret = SteamAPI_Init();
+
+	// look for a headset
+	vr::HmdError HmdErr = vr::HmdError_None;
+	Hmd = vr::VR_Init(&HmdErr);
+	
+	// make sure that the version of the HMD we're compiled against is correct
+	Hmd = (vr::IVRSystem*)vr::VR_GetGenericInterface(vr::IVRSystem_Version, &HmdErr);	//@todo steamvr: verify init error handling
+
+	// initialize SteamController().  Note: this is redundant with the SteamController support.  Verify that the API can handle multiple initializations/teardowns
+	if (SteamController() != nullptr)
+	{
+		FString PluginsDir = FPaths::EnginePluginsDir();
+		FString ContentDir = FPaths::Combine(*PluginsDir, TEXT("Runtime"), TEXT("Steam"), TEXT("SteamController"), TEXT("Content"));
+		FString VdfPath = FPaths::ConvertRelativePathToFull(FPaths::Combine(*ContentDir, TEXT("Controller.vdf")));
+
+		bool bInited = SteamController()->Init(TCHAR_TO_ANSI(*VdfPath));
+		UE_LOG(LogHMD, Log, TEXT("SteamController %s initialized with vdf file '%s'."), bInited ? TEXT("could not be") : TEXT("has been"), *VdfPath);
+	}
+
+	// attach to the compositor
+	if ((Hmd != nullptr) && (HmdErr == vr::HmdError_None))
+	{
+		VRCompositor = (vr::IVRCompositor*)vr::VR_GetGenericInterface(vr::IVRCompositor_Version, &HmdErr);
+
+		if ((VRCompositor != nullptr) && (HmdErr == vr::HmdError_None))
+		{
+			// determine our compositor type
+			vr::Compositor_DeviceType CompositorDeviceType = vr::Compositor_DeviceType_None;
+			if (IsPCPlatform(GMaxRHIShaderPlatform) && !IsOpenGLPlatform(GMaxRHIShaderPlatform))
+			{
+				CompositorDeviceType = vr::Compositor_DeviceType_D3D11;
+			}
+			else if (IsOpenGLPlatform(GMaxRHIShaderPlatform))
+			{
+				check(0);	//@todo steamvr: use old path for mac and linux until support is added
+				CompositorDeviceType = vr::Compositor_DeviceType_OpenGL;
+			}
+		}
+	}
+
+	if ((Hmd != nullptr) && (HmdErr == vr::HmdError_None))
+	{
+		// grab info about the attached display
+		char Buf[128];
+		FString DriverId;
+		vr::TrackedPropertyError Error;
+
+		Hmd->GetStringTrackedDeviceProperty(vr::k_unTrackedDeviceIndex_Hmd, vr::Prop_TrackingSystemName_String, Buf, sizeof(Buf), &Error);
+		if (Error == vr::TrackedProp_Success)
+		{
+			DriverId = FString(UTF8_TO_TCHAR(Buf));
+		}
+
+		Hmd->GetStringTrackedDeviceProperty(vr::k_unTrackedDeviceIndex_Hmd, vr::Prop_SerialNumber_String, Buf, sizeof(Buf), &Error);
+		if (Error == vr::TrackedProp_Success)
+		{
+			DisplayId = FString(UTF8_TO_TCHAR(Buf));
+		}
+
+		// determine our ideal screen percentage
+		uint32 RecommendedWidth, RecommendedHeight;
+		Hmd->GetRecommendedRenderTargetSize(&RecommendedWidth, &RecommendedHeight);
+		RecommendedWidth *= 2;
+
+		int32 ScreenX, ScreenY;
+		uint32 ScreenWidth, ScreenHeight;
+		Hmd->GetWindowBounds(&ScreenX, &ScreenY, &ScreenWidth, &ScreenHeight);
+
+		float WidthPercentage = ((float)RecommendedWidth / (float)ScreenWidth) * 100.0f;
+		float HeightPercentage = ((float)RecommendedHeight / (float)ScreenHeight) * 100.0f;
+
+		float IdealScreenPercentage = FMath::Max(WidthPercentage, HeightPercentage);
+
+		//@todo steamvr: move out of here
+		static IConsoleVariable* CScrPercVar = IConsoleManager::Get().FindConsoleVariable(TEXT("r.ScreenPercentage"));
+
+		if (FMath::RoundToInt(CScrPercVar->GetFloat()) != FMath::RoundToInt(IdealScreenPercentage))
+		{
+			CScrPercVar->Set(IdealScreenPercentage);
+		}
+
+		// disable vsync
+		static IConsoleVariable* CVSyncVar = IConsoleManager::Get().FindConsoleVariable(TEXT("r.VSync"));
+		CVSyncVar->Set(false);
+
+		// enforce finishcurrentframe
+		static IConsoleVariable* CFCFVar = IConsoleManager::Get().FindConsoleVariable(TEXT("r.finishcurrentframe"));
+		CFCFVar->Set(true);
+
+		// Uncap fps to enable FPS higher than 62
+		GEngine->bSmoothFrameRate = false;
+
+		// Grab the chaperone
+		vr::HmdError ChaperoneErr = vr::HmdError_None;
+		VRChaperone = (vr::IVRChaperone*)vr::VR_GetGenericInterface(vr::IVRChaperone_Version, &ChaperoneErr);
+		if ((VRChaperone != nullptr) && (ChaperoneErr == vr::HmdError_None))
+		{
+			ChaperoneBounds = FChaperoneBounds(VRChaperone);
+		}
+		else
+		{
+			UE_LOG(LogHMD, Warning, TEXT("Failed to initialize Chaperone.  Error: %d"), (int32)ChaperoneErr);
+		}
+
+		// Initialize the controller to tracked device mappings
+		InitializeControllerMappings();
+
+#if PLATFORM_WINDOWS
+		if (IsPCPlatform(GMaxRHIShaderPlatform) && !IsOpenGLPlatform(GMaxRHIShaderPlatform))
+		{
+			pD3D11Bridge = new D3D11Bridge(this);
+		}
+#endif
+
+		LoadFromIni();
+
+		UE_LOG(LogHMD, Log, TEXT("SteamVR initialized.  Driver: %s  Display: %s"), *DriverId, *DisplayId);
+	}
+	else
+	{
+		UE_LOG(LogHMD, Warning, TEXT("SteamVR failed to initialize.  Error: %d"), (int32)HmdErr);
+
+		if (Hmd != nullptr)
+		{
+			Hmd = nullptr;
+		}
+	}
+}
+
+void FSteamVRHMD::LoadFromIni()
+{
+	const TCHAR* SteamVRSettings = TEXT("SteamVR.Settings");
+	int32 i;
+
+	if (GConfig->GetInt(SteamVRSettings, TEXT("WindowMirrorMode"), i, GEngineIni))
+	{
+		WindowMirrorMode = i;
+	}
+}
+
+void FSteamVRHMD::SaveToIni()
+{
+	const TCHAR* SteamVRSettings = TEXT("SteamVR.Settings");
+	GConfig->SetInt(SteamVRSettings, TEXT("WindowMirrorMode"), WindowMirrorMode, GEngineIni);
+}
+
+// Initialize the controller to tracked device mappings
+void FSteamVRHMD::InitializeControllerMappings()
+{
+	// invalidate all of the mappings
+	for (uint32 i = 0; i < MAX_STEAM_CONTROLLERS; ++i)
+	{
+		ControllerIndexToTrackedDeviceIndex[i] = -1;
+	}
+
+	if (SteamController() == nullptr)
+	{
+		return;
+	}
+
+	// grab a temporary array of the tracked device ids for comparing later so we're not constantly looping over the tracked devices
+	uint32 TrackedDeviceIds[vr::k_unMaxTrackedDeviceCount];
+	
+	char Buf[128];
+	uint32 BufLen = sizeof(Buf);
+	vr::TrackedPropertyError Error;
+
+	for (uint32 DeviceIndex = 0; DeviceIndex < vr::k_unMaxTrackedDeviceCount; ++DeviceIndex)
+	{
+		Hmd->GetStringTrackedDeviceProperty(DeviceIndex, vr::Prop_AttachedDeviceId_String, Buf, sizeof(Buf), &Error);
+		if (Error == vr::TrackedProp_Success)
+		{
+			TrackedDeviceIds[DeviceIndex] = (uint32)FCStringAnsi::Atoi64(Buf);
+		}
+		else
+		{
+			TrackedDeviceIds[DeviceIndex] = InvalidTrackedDeviceId;
+		}
+	}
+
+	SteamControllerState_t ControllerState;
+
+	// iterate over the controllers
+	for (uint32 ControllerIndex = 0; ControllerIndex < MAX_STEAM_CONTROLLERS; ++ControllerIndex)
+	{
+		if (!SteamController()->GetControllerState(ControllerIndex, &ControllerState))
+		{
+			continue;
+		}
+
+		if (ControllerState.unPacketNum == 0)
+		{
+			continue;
+		}
+
+		uint32 ControllerId = ((uint32)ControllerState.sLeftPadX << 16) | ((uint32)ControllerState.sLeftPadY & 0xFFFF);
+
+		// loop over our cached tracked device ids to see which tracked device index we are mapped to
+		for (uint32 DeviceIndex = 0; DeviceIndex < vr::k_unMaxTrackedDeviceCount; ++DeviceIndex)
+		{
+			if (TrackedDeviceIds[DeviceIndex] == ControllerId)
+			{
+				// found a match, add it to our controller index to tracked device index 
+				ControllerIndexToTrackedDeviceIndex[ControllerIndex] = DeviceIndex;
+				break;
+			}
+		}
+	}
+}
+
+void FSteamVRHMD::Shutdown()
+{
+	// save any runtime configuration changes to the .ini
+	SaveToIni();
+
+	// shut down our headset
+	vr::VR_Shutdown();
+	Hmd = nullptr;
+
+	// shut down the controller interface
+	if (SteamController() != nullptr)
+	{
+		SteamController()->Shutdown();
+	}
+
+	// unload steam library
+	UnloadSteamModule();
+}
+
+bool FSteamVRHMD::LoadSteamModule()
+{
+#if PLATFORM_WINDOWS
+	#if PLATFORM_64BITS
+		FString RootSteamPath;
+		TCHAR VROverridePath[MAX_PATH];
+		FPlatformMisc::GetEnvironmentVariable(TEXT("VR_OVERRIDE"), VROverridePath, MAX_PATH);
+	
+		if (FCString::Strlen(VROverridePath) > 0)
+		{
+			RootSteamPath = FString::Printf(TEXT("%s\\bin\\win64\\"), VROverridePath);
+		}
+		else
+		{
+			RootSteamPath = FPaths::EngineDir() / FString::Printf(TEXT("Binaries/ThirdParty/Steamworks/%s/Win64/"), STEAM_SDK_VER);
+		}
+		
+		FPlatformProcess::PushDllDirectory(*RootSteamPath);
+		SteamDLLHandle = FPlatformProcess::GetDllHandle(*(RootSteamPath + "steam_api64.dll"));
+		FPlatformProcess::PopDllDirectory(*RootSteamPath);
+	#else
+		FString RootSteamPath = FPaths::EngineDir() / FString::Printf(TEXT("Binaries/ThirdParty/Steamworks/%s/Win32/"), STEAM_SDK_VER); 
+		FPlatformProcess::PushDllDirectory(*RootSteamPath);
+		SteamDLLHandle = FPlatformProcess::GetDllHandle(*(RootSteamPath + "steam_api.dll"));
+		FPlatformProcess::PopDllDirectory(*RootSteamPath);
+	#endif
+#elif PLATFORM_MAC
+	SteamDLLHandle = FPlatformProcess::GetDllHandle(TEXT("libsteam_api.dylib"));
+#endif	//PLATFORM_WINDOWS
+
+	if (!SteamDLLHandle)
+	{
+		UE_LOG(LogHMD, Warning, TEXT("Failed to load Steam library."));
+		return false;
+	}
+
+	return true;
+}
+
+void FSteamVRHMD::UnloadSteamModule()
+{
+	if (SteamDLLHandle != nullptr)
+	{
+		FPlatformProcess::FreeDllHandle(SteamDLLHandle);
+		SteamDLLHandle = nullptr;
+	}
+}
+
+#endif //STEAMVR_SUPPORTED_PLATFORMS
diff --git a/Engine/Plugins/Runtime/Steam/SteamVR/Source/SteamVR/Private/SteamVRHMD.h b/Engine/Plugins/Runtime/Steam/SteamVR/Source/SteamVR/Private/SteamVRHMD.h
new file mode 100644
--- /dev/null
+++ b/Engine/Plugins/Runtime/Steam/SteamVR/Source/SteamVR/Private/SteamVRHMD.h
@@ -0,0 +1,386 @@
+// Copyright 1998-2015 Epic Games, Inc. All Rights Reserved.
+
+#pragma once
+#include "ISteamVRPlugin.h"
+#include "HeadMountedDisplay.h"
+#include "IHeadMountedDisplay.h"
+
+#if PLATFORM_WINDOWS
+#include "AllowWindowsPlatformTypes.h"
+#include <d3d11.h>
+#include "HideWindowsPlatformTypes.h"
+#endif
+
+#if STEAMVR_SUPPORTED_PLATFORMS
+
+#include "SceneViewExtension.h"
+
+//@todo steamvr: better association between SteamController plugin and SteamVR plugin
+#ifndef MAX_STEAM_CONTROLLERS
+	#define MAX_STEAM_CONTROLLERS 8
+#endif
+
+/** Stores vectors, in clockwise order, to define soft and hard bounds for Chaperone */
+struct FBoundingQuad
+{
+	FVector Corners[4];
+};
+
+/**
+ * SteamVR Head Mounted Display
+ */
+class FSteamVRHMD : public IHeadMountedDisplay, public ISceneViewExtension
+{
+public:
+	/** IHeadMountedDisplay interface */
+	virtual bool IsHMDConnected() override { return true; }
+	virtual bool IsHMDEnabled() const override;
+	virtual void EnableHMD(bool allow = true) override;
+	virtual EHMDDeviceType::Type GetHMDDeviceType() const override;
+	virtual bool GetHMDMonitorInfo(MonitorInfo&) override;
+
+	virtual void GetFieldOfView(float& OutHFOVInDegrees, float& OutVFOVInDegrees) const override;
+
+	virtual bool DoesSupportPositionalTracking() const override;
+	virtual bool HasValidTrackingPosition() const override;
+	virtual void GetPositionalTrackingCameraProperties(FVector& OutOrigin, FRotator& OutOrientation, float& OutHFOV, float& OutVFOV, float& OutCameraDistance, float& OutNearPlane, float& OutFarPlane) const override;
+
+	virtual void SetInterpupillaryDistance(float NewInterpupillaryDistance) override;
+	virtual float GetInterpupillaryDistance() const override;
+
+	virtual void GetCurrentOrientationAndPosition(FQuat& CurrentOrientation, FVector& CurrentPosition) override;
+	virtual class ISceneViewExtension* GetViewExtension() override;
+	virtual void ApplyHmdRotation(APlayerController* PC, FRotator& ViewRotation) override;
+	virtual void UpdatePlayerCameraRotation(APlayerCameraManager*, struct FMinimalViewInfo& POV) override;
+
+	virtual bool IsChromaAbCorrectionEnabled() const override;
+
+	virtual bool Exec( UWorld* InWorld, const TCHAR* Cmd, FOutputDevice& Ar ) override;
+	virtual void OnScreenModeChange(EWindowMode::Type WindowMode) override;
+
+	virtual bool IsFullscreenAllowed() override { return false; }
+
+	virtual bool IsPositionalTrackingEnabled() const override;
+	virtual bool EnablePositionalTracking(bool enable) override;
+
+	virtual bool IsHeadTrackingAllowed() const override;
+
+	virtual bool IsInLowPersistenceMode() const override;
+	virtual void EnableLowPersistenceMode(bool Enable = true) override;
+
+	virtual void ResetOrientationAndPosition(float yaw = 0.f) override;
+	virtual void ResetOrientation(float Yaw = 0.f) override;
+	virtual void ResetPosition() override;
+
+	virtual void SetClippingPlanes(float NCP, float FCP) override;
+
+	virtual void SetBaseRotation(const FRotator& BaseRot) override;
+	virtual FRotator GetBaseRotation() const override;
+
+	virtual void SetBaseOrientation(const FQuat& BaseOrient) override;
+	virtual FQuat GetBaseOrientation() const override;
+
+	virtual void SetPositionOffset(const FVector& PosOff) override;
+	virtual FVector GetPositionOffset() const override;
+
+	virtual void DrawDistortionMesh_RenderThread(struct FRenderingCompositePassContext& Context, const FSceneView& View, const FIntPoint& TextureSize) override;
+
+	virtual void UpdateScreenSettings(const FViewport* InViewport) {}
+
+	/** IStereoRendering interface */
+	virtual bool IsStereoEnabled() const override;
+	virtual bool EnableStereo(bool stereo = true) override;
+	virtual void AdjustViewRect(EStereoscopicPass StereoPass, int32& X, int32& Y, uint32& SizeX, uint32& SizeY) const override;
+	virtual void CalculateStereoViewOffset(const EStereoscopicPass StereoPassType, const FRotator& ViewRotation, 
+										   const float MetersToWorld, FVector& ViewLocation) override;
+	virtual FMatrix GetStereoProjectionMatrix(const EStereoscopicPass StereoPassType, const float FOV) const override;
+	virtual void InitCanvasFromView(FSceneView* InView, UCanvas* Canvas) override;
+	virtual void PushViewportCanvas(EStereoscopicPass StereoPass, FCanvas *InCanvas, UCanvas *InCanvasObject, FViewport *InViewport) const override;
+	virtual void PushViewCanvas(EStereoscopicPass StereoPass, FCanvas *InCanvas, UCanvas *InCanvasObject, FSceneView *InView) const override;
+	virtual void GetEyeRenderParams_RenderThread(EStereoscopicPass StereoPass, FVector2D& EyeToSrcUVScaleValue, FVector2D& EyeToSrcUVOffsetValue) const override;
+	virtual void RenderTexture_RenderThread(FRHICommandListImmediate& RHICmdList, FTexture2DRHIParamRef BackBuffer, FTexture2DRHIParamRef SrcTexture) const override;
+
+	virtual void CalculateRenderTargetSize(uint32& InOutSizeX, uint32& InOutSizeY) const override;
+	virtual bool NeedReAllocateViewportRenderTarget(const FViewport& Viewport) const override;
+	virtual bool ShouldUseSeparateRenderTarget() const override
+	{
+		check(IsInGameThread());
+		return IsStereoEnabled();
+	}
+
+	/** ISceneViewExtension interface */
+	virtual void ModifyShowFlags(FEngineShowFlags& ShowFlags) override;
+	virtual void SetupView(FSceneViewFamily& InViewFamily, FSceneView& InView) override;
+	virtual void PreRenderView_RenderThread(FSceneView& InView) override;
+	virtual void PreRenderViewFamily_RenderThread(FSceneViewFamily& InViewFamily) override;
+
+	virtual void UpdateViewport(bool bUseSeparateRenderTarget, const FViewport& Viewport, SViewport*) override;
+
+	class BridgeBaseImpl : public FRHICustomPresent
+	{
+	public:
+		BridgeBaseImpl(FSteamVRHMD* plugin) :
+			FRHICustomPresent(nullptr),
+			Plugin(plugin), 
+			bNeedReinitRendererAPI(true), 
+			bInitialized(false) 
+		{}
+
+		bool IsInitialized() const { return bInitialized; }
+
+		virtual void BeginRendering() = 0;
+		virtual void UpdateViewport(const FViewport& Viewport, FRHIViewport* ViewportRHI) = 0;
+		virtual void SetNeedReinitRendererAPI() { bNeedReinitRendererAPI = true; }
+
+		virtual void Reset() = 0;
+		virtual void Shutdown() = 0;
+
+	protected:
+		FSteamVRHMD*		Plugin;
+		bool				bNeedReinitRendererAPI;
+		bool				bInitialized;
+	};
+
+#if PLATFORM_WINDOWS
+	class D3D11Bridge : public BridgeBaseImpl
+	{
+	public:
+		D3D11Bridge(FSteamVRHMD* plugin);
+
+		virtual void OnBackBufferResize() override;
+		virtual bool Present(int SyncInterval) override;
+
+		virtual void BeginRendering() override;
+		void FinishRendering();
+		virtual void UpdateViewport(const FViewport& Viewport, FRHIViewport* ViewportRHI) override;
+		virtual void Reset() override;
+		virtual void Shutdown() override
+		{
+			Reset();
+		}
+
+	protected:
+		ID3D11Texture2D* RenderTargetTexture = NULL;
+	};
+#endif // PLATFORM_WINDOWS
+
+	BridgeBaseImpl* GetActiveRHIBridgeImpl();
+	void ShutdownRendering();
+
+	/** Motion Controllers */
+	void GetTrackedDeviceIds(TArray<int32>& TrackedIds);
+	bool GetTrackedObjectOrientationAndPosition(uint32 DeviceId, FQuat& CurrentOrientation, FVector& CurrentPosition);
+	bool GetTrackedDeviceIdFromControllerIndex(int32 ControllerIndex, int32& OutDeviceId);
+
+	/** Chaperone */
+	/** Returns whether or not the player is currently inside the soft bounds */
+	bool IsInsideSoftBounds();
+
+	/** Returns an array of the soft bounds as Unreal-scaled vectors, relative to the HMD calibration point (0,0,0).  The Z will always be at 0.f */
+	TArray<FVector> GetSoftBounds() const;
+
+	/** Returns an array of the hard bounds as Unreal-scaled vectors, relative to the HMD calibration point (0,0,0).  Every four entries will represent one quad of the bounding box */
+	TArray<FVector> GetHardBounds() const;
+
+	/** Get the windowed mirror mode.  @todo steamvr: thread safe flags */
+	int32 GetWindowMirrorMode() const { return WindowMirrorMode; }
+
+public:
+	/** Constructor */
+	FSteamVRHMD();
+
+	/** Destructor */
+	virtual ~FSteamVRHMD();
+
+	/** @return	True if the API was initialized OK */
+	bool IsInitialized() const;
+
+private:
+
+	/**
+	 * Starts up the SteamVR API
+	 */
+	void Startup();
+
+	/**
+	 * Shuts down the SteamVR API
+	 */
+	void Shutdown();
+
+	void LoadFromIni();
+	void SaveToIni();
+
+	bool LoadSteamModule();
+	void UnloadSteamModule();
+
+	// Initialize the controller to tracked device mappings
+	void InitializeControllerMappings();
+
+	void PoseToOrientationAndPosition(const vr::HmdMatrix34_t& Pose, FQuat& OutOrientation, FVector& OutPosition) const;
+	void GetCurrentPose(FQuat& CurrentOrientation, FVector& CurrentPosition, uint32 DeviceID = vr::k_unTrackedDeviceIndex_Hmd, bool bForceRefresh=false);
+
+	FORCEINLINE FMatrix ToFMatrix(const vr::HmdMatrix34_t& tm) const
+	{
+		// Rows and columns are swapped between vr::HmdMatrix34_t and FMatrix
+		return FMatrix(
+			FPlane(tm.m[0][0], tm.m[1][0], tm.m[2][0], 0.0f),
+			FPlane(tm.m[0][1], tm.m[1][1], tm.m[2][1], 0.0f),
+			FPlane(tm.m[0][2], tm.m[1][2], tm.m[2][2], 0.0f),
+			FPlane(tm.m[0][3], tm.m[1][3], tm.m[2][3], 1.0f));
+	}
+
+	FORCEINLINE FMatrix ToFMatrix(const vr::HmdMatrix44_t& tm) const
+	{
+		// Rows and columns are swapped between vr::HmdMatrix44_t and FMatrix
+		return FMatrix(
+			FPlane(tm.m[0][0], tm.m[1][0], tm.m[2][0], tm.m[3][0]),
+			FPlane(tm.m[0][1], tm.m[1][1], tm.m[2][1], tm.m[3][1]),
+			FPlane(tm.m[0][2], tm.m[1][2], tm.m[2][2], tm.m[3][2]),
+			FPlane(tm.m[0][3], tm.m[1][3], tm.m[2][3], tm.m[3][3]));
+	}
+
+private:
+
+	bool bHmdEnabled;
+	bool bStereoEnabled;
+	bool bHmdPosTracking;
+	mutable bool bHaveVisionTracking;
+
+ 	struct FTrackingFrame
+ 	{
+ 		uint32 FrameNumber;
+ 
+ 		bool bDeviceIsValid[vr::k_unMaxTrackedDeviceCount];
+ 		bool bPoseIsValid[vr::k_unMaxTrackedDeviceCount];
+ 		FVector DevicePosition[vr::k_unMaxTrackedDeviceCount];
+ 		FQuat DeviceOrientation[vr::k_unMaxTrackedDeviceCount];
+
+		vr::HmdMatrix34_t RawPoses[vr::k_unMaxTrackedDeviceCount];
+
+		FTrackingFrame()
+		{
+			FrameNumber = 0;
+
+			const uint32 MaxDevices = vr::k_unMaxTrackedDeviceCount;
+
+			FMemory::Memzero(bDeviceIsValid, MaxDevices * sizeof(bool));
+			FMemory::Memzero(bPoseIsValid, MaxDevices * sizeof(bool));
+			FMemory::Memzero(DevicePosition, MaxDevices * sizeof(FVector));
+
+			for (uint32 i = 0; i < vr::k_unMaxTrackedDeviceCount; ++i)
+			{
+				DeviceOrientation[i] = FQuat::Identity;
+			}
+
+			FMemory::Memzero(RawPoses, MaxDevices * sizeof(vr::HmdMatrix34_t));
+		}
+ 	};
+	FTrackingFrame TrackingFrame;
+
+	// a mapping from steam controller index to a tracked device index.  -1 indicates an invalid mapping
+	int32 ControllerIndexToTrackedDeviceIndex[MAX_STEAM_CONTROLLERS];
+	static const uint32 InvalidTrackedDeviceId = 0xFFFFFFFF;
+
+	/** Coverts a SteamVR-space vector to an Unreal-space vector.  Does not handle scaling, only axes conversion */
+	FORCEINLINE static FVector CONVERT_STEAMVECTOR_TO_FVECTOR(const vr::HmdVector3_t InVector)
+	{
+		return FVector(-InVector.v[2], InVector.v[0], InVector.v[1]);
+	}
+
+	FORCEINLINE static FVector RAW_STEAMVECTOR_TO_FVECTOR(const vr::HmdVector3_t InVector)
+	{
+		return FVector(InVector.v[0], InVector.v[1], InVector.v[2]);
+	}
+
+	/** Chaperone Support */
+	struct FChaperoneBounds
+	{
+		/** Stores the soft bounds in SteamVR HMD space, for fast checking.  These will need to be converted to Unreal HMD-calibrated space before being used in the world */
+		FBoundingQuad			SoftBounds;
+
+		/** Stores the hard bounds in SteamVR HMD space, for fast checking.  These will need to be converted to Unreal HMD-calibrated space before being used in the world */
+		TArray<FBoundingQuad>	HardBounds;
+
+		uint32 NumHardBounds;
+
+	public:
+		FChaperoneBounds()
+			: NumHardBounds(0)
+		{}
+
+		FChaperoneBounds(vr::IVRChaperone* Chaperone)
+			: FChaperoneBounds()
+		{
+			vr::ChaperoneSoftBoundsInfo_t	VRSoftBounds;
+			Chaperone->GetSoftBoundsInfo(&VRSoftBounds);
+			for (uint8 i = 0; i < 4; ++i)
+			{
+				const vr::HmdVector3_t Corner = VRSoftBounds.quadCorners.vCorners[i];
+				SoftBounds.Corners[i] = RAW_STEAMVECTOR_TO_FVECTOR(Corner);
+			}
+
+			// Check to see if we have any bounds specified.  This MUST be called before actually getting buffer info
+			Chaperone->GetHardBoundsInfo(NULL, &NumHardBounds);
+			if (NumHardBounds > 0)
+			{
+				vr::HmdQuad_t* HardBoundsBuf = (vr::HmdQuad_t*)FMemory_Alloca(NumHardBounds * sizeof(vr::HmdQuad_t));
+				FMemory::Memzero(HardBoundsBuf, NumHardBounds * sizeof(vr::HmdQuad_t));
+
+				// Actually grab the bounds from the buffer
+				Chaperone->GetHardBoundsInfo(HardBoundsBuf, &NumHardBounds);
+				for (uint32 i = 0; i < NumHardBounds; ++i)
+				{
+					FBoundingQuad Quad;
+					Quad.Corners[0] = RAW_STEAMVECTOR_TO_FVECTOR(HardBoundsBuf[i].vCorners[0]);
+					Quad.Corners[1] = RAW_STEAMVECTOR_TO_FVECTOR(HardBoundsBuf[i].vCorners[1]);
+					Quad.Corners[2] = RAW_STEAMVECTOR_TO_FVECTOR(HardBoundsBuf[i].vCorners[2]);
+					Quad.Corners[3] = RAW_STEAMVECTOR_TO_FVECTOR(HardBoundsBuf[i].vCorners[3]);
+
+					HardBounds.Add(Quad);
+				}
+			}
+		}
+	};
+	FChaperoneBounds ChaperoneBounds;
+	
+	float IPD;
+	int32 WindowMirrorMode;		// how to mirror the display contents to the desktop window: 0 - no mirroring, 1 - single eye, 2 - stereo pair
+
+	/** Player's orientation tracking */
+	mutable FQuat			CurHmdOrientation;
+
+	FRotator				DeltaControlRotation;    // same as DeltaControlOrientation but as rotator
+	FQuat					DeltaControlOrientation; // same as DeltaControlRotation but as quat
+
+	mutable FVector			CurHmdPosition;
+
+	mutable FQuat			LastHmdOrientation; // contains last APPLIED ON GT HMD orientation
+	FVector					LastHmdPosition;	// contains last APPLIED ON GT HMD position
+
+	// HMD base values, specify forward orientation and zero pos offset
+	FQuat					BaseOrientation;	// base orientation
+
+	/** World units (UU) to Meters scale.  Read from the level, and used to transform positional tracking data */
+	float WorldToMetersScale;
+
+	IRendererModule* RendererModule;
+
+	vr::IVRSystem* Hmd;
+	vr::IVRCompositor* VRCompositor;
+	vr::IVRChaperone* VRChaperone;
+
+	void* SteamDLLHandle;
+	
+	FString DisplayId;
+
+	float IdealScreenPercentage;
+
+#if PLATFORM_WINDOWS
+	TRefCountPtr<D3D11Bridge>	pD3D11Bridge;
+#endif
+};
+
+
+DEFINE_LOG_CATEGORY_STATIC(LogHMD, Log, All);
+
+#endif //STEAMVR_SUPPORTED_PLATFORMS
diff --git a/Engine/Plugins/Runtime/Steam/SteamVR/Source/SteamVR/Private/SteamVRPrivatePCH.h b/Engine/Plugins/Runtime/Steam/SteamVR/Source/SteamVR/Private/SteamVRPrivatePCH.h
new file mode 100644
--- /dev/null
+++ b/Engine/Plugins/Runtime/Steam/SteamVR/Source/SteamVR/Private/SteamVRPrivatePCH.h
@@ -0,0 +1,34 @@
+// Copyright 1998-2015 Epic Games, Inc. All Rights Reserved.
+
+#pragma once
+
+/** Name of the current Steam SDK version in use (matches directory name) */
+#define STEAM_SDK_VER TEXT("Steamv132")
+
+#include "Engine.h"
+#include "IHeadMountedDisplay.h"
+#include "Runtime/Engine/Public/ScreenRendering.h"
+
+// @todo Steam: Steam headers trigger secure-C-runtime warnings in Visual C++. Rather than mess with _CRT_SECURE_NO_WARNINGS, we'll just
+//	disable the warnings locally. Remove when this is fixed in the SDK
+#if WITH_STEAMWORKS
+
+#ifdef _MSC_VER
+#pragma warning(push)
+#pragma warning(disable:4996)
+#endif
+
+#pragma push_macro("ARRAY_COUNT")
+#undef ARRAY_COUNT
+
+#include "steam/steam_api.h"
+#include "steam/steamvr.h"
+
+#pragma pop_macro("ARRAY_COUNT")
+
+// @todo Steam: See above
+#ifdef _MSC_VER
+#pragma warning(pop)
+#endif
+
+#endif // WITH_STEAMWORKS
\ No newline at end of file
diff --git a/Engine/Plugins/Runtime/Steam/SteamVR/Source/SteamVR/Private/SteamVRRender.cpp b/Engine/Plugins/Runtime/Steam/SteamVR/Source/SteamVR/Private/SteamVRRender.cpp
new file mode 100644
--- /dev/null
+++ b/Engine/Plugins/Runtime/Steam/SteamVR/Source/SteamVR/Private/SteamVRRender.cpp
@@ -0,0 +1,156 @@
+// Copyright 1998-2015 Epic Games, Inc. All Rights Reserved.
+//
+#include "SteamVRPrivatePCH.h"
+#include "SteamVRHMD.h"
+
+#include "RendererPrivate.h"
+#include "ScenePrivate.h"
+#include "PostProcess/PostProcessHMD.h"
+
+#if STEAMVR_SUPPORTED_PLATFORMS
+
+void FSteamVRHMD::DrawDistortionMesh_RenderThread(struct FRenderingCompositePassContext& Context, const FSceneView& View, const FIntPoint& TextureSize)
+{
+	check(0);
+}
+
+void FSteamVRHMD::RenderTexture_RenderThread(FRHICommandListImmediate& RHICmdList, FTexture2DRHIParamRef BackBuffer, FTexture2DRHIParamRef SrcTexture) const
+{
+	check(IsInRenderingThread());
+
+	if (WindowMirrorMode == 0)
+	{
+		return;
+	}
+
+	const uint32 ViewportWidth = BackBuffer->GetSizeX();
+	const uint32 ViewportHeight = BackBuffer->GetSizeY();
+
+	SetRenderTarget(RHICmdList, BackBuffer, FTextureRHIRef());
+	RHICmdList.SetViewport(0, 0, 0, ViewportWidth, ViewportHeight, 1.0f);
+
+	RHICmdList.SetBlendState(TStaticBlendState<>::GetRHI());
+	RHICmdList.SetRasterizerState(TStaticRasterizerState<>::GetRHI());
+	RHICmdList.SetDepthStencilState(TStaticDepthStencilState<false, CF_Always>::GetRHI());
+
+	const auto FeatureLevel = GMaxRHIFeatureLevel;
+	auto ShaderMap = GetGlobalShaderMap(FeatureLevel);
+
+	TShaderMapRef<FScreenVS> VertexShader(ShaderMap);
+	TShaderMapRef<FScreenPS> PixelShader(ShaderMap);
+
+	static FGlobalBoundShaderState BoundShaderState;
+	SetGlobalBoundShaderState(RHICmdList, FeatureLevel, BoundShaderState, RendererModule->GetFilterVertexDeclaration().VertexDeclarationRHI, *VertexShader, *PixelShader);
+
+	PixelShader->SetParameters(RHICmdList, TStaticSamplerState<SF_Bilinear>::GetRHI(), SrcTexture);
+
+	if (WindowMirrorMode == 1)
+	{
+		// need to clear when rendering only one eye since the borders won't be touched by the DrawRect below
+		RHICmdList.Clear(true, FLinearColor::Black, false, 0, false, 0, FIntRect());
+
+		RendererModule->DrawRectangle(
+			RHICmdList,
+			ViewportWidth / 4, 0,
+			ViewportWidth / 2, ViewportHeight,
+			0.1f, 0.2f,
+			0.3f, 0.6f,
+			FIntPoint(ViewportWidth, ViewportHeight),
+			FIntPoint(1, 1),
+			*VertexShader,
+			EDRF_Default);
+	}
+	else if (WindowMirrorMode == 2)
+	{
+		RendererModule->DrawRectangle(
+			RHICmdList,
+			0, 0,
+			ViewportWidth, ViewportHeight,
+			0.0f, 0.0f,
+			1.0f, 1.0f,
+			FIntPoint(ViewportWidth, ViewportHeight),
+			FIntPoint(1, 1),
+			*VertexShader,
+			EDRF_Default);
+	}
+}
+
+#if PLATFORM_WINDOWS
+
+FSteamVRHMD::D3D11Bridge::D3D11Bridge(FSteamVRHMD* plugin):
+	BridgeBaseImpl(plugin),
+	RenderTargetTexture(nullptr)
+{
+}
+
+void FSteamVRHMD::D3D11Bridge::BeginRendering()
+{
+	check(IsInRenderingThread());
+
+	static bool Inited = false;
+	if (!Inited)
+	{
+		Plugin->VRCompositor->SetGraphicsDevice(vr::Compositor_DeviceType_D3D11, RHIGetNativeDevice());
+		Inited = true;
+	}
+}
+
+void FSteamVRHMD::D3D11Bridge::FinishRendering()
+{
+	vr::Compositor_TextureBounds LeftBounds;
+	LeftBounds.uMin = 0.0f;
+	LeftBounds.uMax = 0.5f;
+	LeftBounds.vMin = 0.0f;
+	LeftBounds.vMax = 1.0f;
+
+	Plugin->VRCompositor->Submit(vr::Eye_Left, RenderTargetTexture, &LeftBounds);
+
+	vr::Compositor_TextureBounds RightBounds;
+	RightBounds.uMin = 0.5f;
+	RightBounds.uMax = 1.0f;
+	RightBounds.vMin = 0.0f;
+	RightBounds.vMax = 1.0f;
+
+	Plugin->VRCompositor->Submit(vr::Eye_Right, RenderTargetTexture, &RightBounds);
+}
+
+void FSteamVRHMD::D3D11Bridge::Reset()
+{
+}
+
+void FSteamVRHMD::D3D11Bridge::UpdateViewport(const FViewport& Viewport, FRHIViewport* ViewportRHI)
+{
+	check(IsInGameThread());
+	check(ViewportRHI);
+
+	const FTexture2DRHIRef& RT = Viewport.GetRenderTargetTexture();
+	check(IsValidRef(RT));
+
+	if (RenderTargetTexture != nullptr)
+	{
+		RenderTargetTexture->Release();	//@todo steamvr: need to also release in reset
+	}
+
+	RenderTargetTexture = (ID3D11Texture2D*)RT->GetNativeResource();
+	RenderTargetTexture->AddRef();
+
+	ViewportRHI->SetCustomPresent(this);
+}
+
+
+void FSteamVRHMD::D3D11Bridge::OnBackBufferResize()
+{
+}
+
+bool FSteamVRHMD::D3D11Bridge::Present(int SyncInterval)
+{
+	check(IsInRenderingThread());
+
+	FinishRendering();
+
+	return true;
+}
+
+#endif // PLATFORM_WINDOWS
+
+#endif // STEAMVR_SUPPORTED_PLATFORMS
\ No newline at end of file
diff --git a/Engine/Plugins/Runtime/Steam/SteamVR/Source/SteamVR/Public/ISteamVRPlugin.h b/Engine/Plugins/Runtime/Steam/SteamVR/Source/SteamVR/Public/ISteamVRPlugin.h
new file mode 100644
--- /dev/null
+++ b/Engine/Plugins/Runtime/Steam/SteamVR/Source/SteamVR/Public/ISteamVRPlugin.h
@@ -0,0 +1,40 @@
+// Copyright 1998-2015 Epic Games, Inc. All Rights Reserved.
+
+#pragma once
+
+#include "ModuleManager.h"
+#include "IHeadMountedDisplayModule.h"
+
+#define STEAMVR_SUPPORTED_PLATFORMS (PLATFORM_WINDOWS && WITH_STEAMWORKS)
+
+/**
+ * The public interface to this module.  In most cases, this interface is only public to sibling modules 
+ * within this plugin.
+ */
+class ISteamVRPlugin : public IHeadMountedDisplayModule
+{
+
+public:
+
+	/**
+	 * Singleton-like access to this module's interface.  This is just for convenience!
+	 * Beware of calling this during the shutdown phase, though.  Your module might have been unloaded already.
+	 *
+	 * @return Returns singleton instance, loading the module on demand if needed
+	 */
+	static inline ISteamVRPlugin& Get()
+	{
+		return FModuleManager::LoadModuleChecked< ISteamVRPlugin >( "SteamVR" );
+	}
+
+	/**
+	 * Checks to see if this module is loaded and ready.  It is only valid to call Get() if IsAvailable() returns true.
+	 *
+	 * @return True if the module is loaded and ready to use
+	 */
+	static inline bool IsAvailable()
+	{
+		return FModuleManager::Get().IsModuleLoaded( "SteamVR" );
+	}
+};
+
diff --git a/Engine/Plugins/Runtime/Steam/SteamVR/Source/SteamVR/SteamVR.Build.cs b/Engine/Plugins/Runtime/Steam/SteamVR/Source/SteamVR/SteamVR.Build.cs
new file mode 100644
--- /dev/null
+++ b/Engine/Plugins/Runtime/Steam/SteamVR/Source/SteamVR/SteamVR.Build.cs
@@ -0,0 +1,42 @@
+// Copyright 1998-2015 Epic Games, Inc. All Rights Reserved.
+
+namespace UnrealBuildTool.Rules
+{
+	public class SteamVR : ModuleRules
+	{
+		public SteamVR(TargetInfo Target)
+		{
+			PrivateIncludePaths.AddRange(
+				new string[] {
+					"SteamVR/Private",
+					"../../../../../Source/Runtime/Renderer/Private",
+					// ... add other private include paths required here ...
+				}
+				);
+
+			PrivateDependencyModuleNames.AddRange(
+				new string[]
+				{
+					"Core",
+					"CoreUObject",
+					"Engine",
+					"RHI",
+					"RenderCore",
+					"Renderer",
+					"ShaderCore",
+					"HeadMountedDisplay",
+					"Slate",
+					"SlateCore"
+				}
+				);
+
+			Definitions.Add("WITH_STEAMWORKS=" + (UEBuildConfiguration.bCompileSteamOSS ? "1" : "0"));
+			
+            if (UEBuildConfiguration.bCompileSteamOSS && (Target.Platform == UnrealTargetPlatform.Win32 || Target.Platform == UnrealTargetPlatform.Win64))
+            {
+				AddThirdPartyPrivateStaticDependencies(Target, "Steamworks");
+                PrivateDependencyModuleNames.AddRange(new string[] { "D3D11RHI" });     //@todo steamvr: multiplatform
+            }
+		}
+	}
+}
diff --git a/Engine/Plugins/Runtime/Steam/SteamVR/SteamVR.uplugin b/Engine/Plugins/Runtime/Steam/SteamVR/SteamVR.uplugin
new file mode 100644
--- /dev/null
+++ b/Engine/Plugins/Runtime/Steam/SteamVR/SteamVR.uplugin
@@ -0,0 +1,22 @@
+{
+	"FileVersion" : 3,
+
+	"FriendlyName" : "SteamVR Plugin",
+	"Version" : 1,
+	"VersionName" : "1.0",
+	"CreatedBy" : "Epic Games, Inc.",
+	"CreatedByURL" : "http://epicgames.com",
+	"EngineVersion" : "4.2.0",
+	"Description" : "Support for the SteamVR API",
+	"Category" : "Head Mounted Displays.Devices",
+	"EnabledByDefault" : false,
+
+	"Modules" :
+	[
+		{
+			"Name" : "SteamVR",
+			"Type" : "Runtime",
+			"WhitelistPlatforms" : [ "Win64", "Win32" ]
+		}
+	]
+}
