- Avoid a memory overwrite in SteamVR if more than 8 controllers are attached

diff --git a/Engine/Plugins/Runtime/Steam/SteamVR/Source/SteamVRController/Private/SteamVRController.cpp b/Engine/Plugins/Runtime/Steam/SteamVR/Source/SteamVRController/Private/SteamVRController.cpp
--- a/Engine/Plugins/Runtime/Steam/SteamVR/Source/SteamVRController/Private/SteamVRController.cpp
+++ b/Engine/Plugins/Runtime/Steam/SteamVR/Source/SteamVRController/Private/SteamVRController.cpp
@@ -91,6 +91,12 @@ public:
 				// update the mappings if this is a new device
 				if (DeviceToControllerMap[DeviceIndex] == -1)
 				{
+					// don't map too many controllers
+					if (NumControllersMapped >= MAX_STEAMVR_CONTROLLERS)
+					{
+						continue;
+					}
+
 					DeviceToControllerMap[DeviceIndex] = NumControllersMapped;
 					ControllerToDeviceMap[NumControllersMapped] = DeviceIndex;
 					++NumControllersMapped;
