#steamvr Adding support for sitting/standing, updating device id functions to gather ids by class

diff --git a/Engine/Plugins/Runtime/Steam/SteamVR/Source/SteamVR/Classes/SteamVRFunctionLibrary.h b/Engine/Plugins/Runtime/Steam/SteamVR/Source/SteamVR/Classes/SteamVRFunctionLibrary.h
--- a/Engine/Plugins/Runtime/Steam/SteamVR/Source/SteamVR/Classes/SteamVRFunctionLibrary.h
+++ b/Engine/Plugins/Runtime/Steam/SteamVR/Source/SteamVR/Classes/SteamVRFunctionLibrary.h
@@ -3,6 +3,40 @@
 #pragma once
 #include "SteamVRFunctionLibrary.generated.h"
 
+/** Describes the type of device tracked in the SteamVR volume (e.g. controller, Lighthouse, etc) */
+UENUM(BlueprintType)
+namespace ESteamVRTrackedDeviceType
+{
+	enum Type
+	{
+		/** Represents a Steam VR Controller */
+		Controller,
+
+		/** Represents a static tracking reference device, such as a Lighthouse or tracking camera */
+		TrackingReference,
+
+		/** Misc. device types, for future expansion */
+		Other,
+
+		/** DeviceId is invalid */
+		Invalid
+	};
+}
+
+/** Defines which calibration point to use (e.g. center of the room for standing, calibrated head position for seated) */
+UENUM(BlueprintType)
+namespace ESteamVRTrackingSpace
+{
+	enum Type
+	{
+		/** Standing origin, where poses are relative to the safe room bounds provided by the user's calibration */
+		Standing,
+
+		/** Seated origin, where poses are relative to the user's calibrated seated head position */
+		Seated
+	};
+}
+
 /**
  * SteamVR Extensions Function Library
  */
@@ -14,10 +48,11 @@ class USteamVRFunctionLibrary : public UBlueprintFunctionLibrary
 	/**
 	 * Returns an array of the currently tracked device IDs
 	 *
+	 * @param	DeviceType	Which class of device (e.g. controller, tracking devices) to get Device Ids for
 	 * @param	OutTrackedDeviceIds	(out) Array containing the ID of each device that's currently tracked
 	 */
 	UFUNCTION(BlueprintPure, Category="SteamVR")
-	static void GetValidTrackedDeviceIds(TArray<int32>& OutTrackedDeviceIds);
+	static void GetValidTrackedDeviceIds(TEnumAsByte<ESteamVRTrackedDeviceType::Type> DeviceType, TArray<int32>& OutTrackedDeviceIds);
 
 	/**
 	 * Gets the orientation and position (in device space) of the device with the specified ID
@@ -39,4 +74,18 @@ class USteamVRFunctionLibrary : public UBlueprintFunctionLibrary
 	 */
 	UFUNCTION(BlueprintPure, Category = "SteamVR")
 	static bool GetTrackedDeviceIdFromControllerIndex(int32 ControllerIndex, int32& OutDeviceId);
+
+	/**
+	 * Sets the tracking space (e.g. sitting or standing), which changes which space tracked positions are returned to
+	 *
+	 * NewSpace		The new space to consider all tracked positions in.  For instance, standing assumes the center of the safe zone as the origin.
+	 */
+	UFUNCTION(BlueprintCallable, Category = "SteamVR")
+	static void SetTrackingSpace(TEnumAsByte<ESteamVRTrackingSpace::Type> NewSpace);
+
+	/**
+	 * Gets the tracking space (e.g. sitting or standing), which determines the location of the origin.
+	 */
+	UFUNCTION(BlueprintPure, Category = "SteamVR")
+	static TEnumAsByte<ESteamVRTrackingSpace::Type> GetTrackingSpace();
 };
diff --git a/Engine/Plugins/Runtime/Steam/SteamVR/Source/SteamVR/Private/SteamVRFunctionLibrary.cpp b/Engine/Plugins/Runtime/Steam/SteamVR/Source/SteamVR/Private/SteamVRFunctionLibrary.cpp
--- a/Engine/Plugins/Runtime/Steam/SteamVR/Source/SteamVR/Private/SteamVRFunctionLibrary.cpp
+++ b/Engine/Plugins/Runtime/Steam/SteamVR/Source/SteamVR/Private/SteamVRFunctionLibrary.cpp
@@ -9,14 +9,24 @@ USteamVRFunctionLibrary::USteamVRFunctionLibrary(const FObjectInitializer& Objec
 {
 }
 
-void USteamVRFunctionLibrary::GetValidTrackedDeviceIds(TArray<int32>& DeviceIds)
+FSteamVRHMD* GetSteamVRHMD()
 {
-	DeviceIds.Empty();
-
 	if (GEngine->HMDDevice.IsValid() && (GEngine->HMDDevice->GetHMDDeviceType() == EHMDDeviceType::DT_SteamVR))
 	{
-		FSteamVRHMD* SteamVRHMD = static_cast<FSteamVRHMD*>(GEngine->HMDDevice.Get());
-		SteamVRHMD->GetTrackedDeviceIds(DeviceIds);
+		return static_cast<FSteamVRHMD*>(GEngine->HMDDevice.Get());
+	}
+
+	return nullptr;
+}
+
+void USteamVRFunctionLibrary::GetValidTrackedDeviceIds(TEnumAsByte<ESteamVRTrackedDeviceType::Type> DeviceType, TArray<int32>& OutTrackedDeviceIds)
+{
+	OutTrackedDeviceIds.Empty();
+
+	FSteamVRHMD* SteamVRHMD = GetSteamVRHMD();
+	if (SteamVRHMD)
+	{
+		SteamVRHMD->GetTrackedDeviceIds(DeviceType, OutTrackedDeviceIds);
 	}
 }
 
@@ -24,10 +34,9 @@ bool USteamVRFunctionLibrary::GetTrackedDevicePositionAndOrientation(int32 Devic
 {
 	bool RetVal = false;
 
-	if (GEngine->HMDDevice.IsValid() && (GEngine->HMDDevice->GetHMDDeviceType() == EHMDDeviceType::DT_SteamVR))
+	FSteamVRHMD* SteamVRHMD = GetSteamVRHMD();
+	if (SteamVRHMD)
 	{
-		FSteamVRHMD* SteamVRHMD = static_cast<FSteamVRHMD*>(GEngine->HMDDevice.Get());
-
 		FQuat DeviceOrientation = FQuat::Identity;
 		RetVal = SteamVRHMD->GetTrackedObjectOrientationAndPosition(DeviceId, DeviceOrientation, OutPosition);
 		OutOrientation = DeviceOrientation.Rotator();
@@ -40,12 +49,33 @@ bool USteamVRFunctionLibrary::GetTrackedDeviceIdFromControllerIndex(int32 Contro
 {
 	OutDeviceId = -1;
 
-	if (!GEngine->HMDDevice.IsValid() || (GEngine->HMDDevice->GetHMDDeviceType() != EHMDDeviceType::DT_SteamVR))
+	FSteamVRHMD* SteamVRHMD = GetSteamVRHMD();
+	if (SteamVRHMD)
 	{
-		return false;
+		return SteamVRHMD->GetTrackedDeviceIdFromControllerIndex(ControllerIndex, OutDeviceId);
 	}
 
-	FSteamVRHMD* SteamVRHMD = static_cast<FSteamVRHMD*>(GEngine->HMDDevice.Get());
+	return false;
+}
 
-	return SteamVRHMD->GetTrackedDeviceIdFromControllerIndex(ControllerIndex, OutDeviceId);
+void USteamVRFunctionLibrary::SetTrackingSpace(TEnumAsByte<ESteamVRTrackingSpace::Type> NewSpace)
+{
+	FSteamVRHMD* SteamVRHMD = GetSteamVRHMD();
+	if (SteamVRHMD)
+	{
+		SteamVRHMD->SetTrackingSpace(NewSpace);
+	}
 }
+
+TEnumAsByte<ESteamVRTrackingSpace::Type> USteamVRFunctionLibrary::GetTrackingSpace()
+{
+	ESteamVRTrackingSpace::Type RetVal = ESteamVRTrackingSpace::Standing;
+
+	FSteamVRHMD* SteamVRHMD = GetSteamVRHMD();
+	if (SteamVRHMD)
+	{
+		RetVal = SteamVRHMD->GetTrackingSpace();
+	}
+
+	return RetVal;
+}
\ No newline at end of file
diff --git a/Engine/Plugins/Runtime/Steam/SteamVR/Source/SteamVR/Private/SteamVRHMD.cpp b/Engine/Plugins/Runtime/Steam/SteamVR/Source/SteamVR/Private/SteamVRHMD.cpp
--- a/Engine/Plugins/Runtime/Steam/SteamVR/Source/SteamVR/Private/SteamVRHMD.cpp
+++ b/Engine/Plugins/Runtime/Steam/SteamVR/Source/SteamVR/Private/SteamVRHMD.cpp
@@ -251,6 +251,47 @@ TArray<FVector> FSteamVRHMD::GetHardBounds() const
 	return Bounds;
 }
 
+void FSteamVRHMD::SetTrackingSpace(TEnumAsByte<ESteamVRTrackingSpace::Type> NewSpace)
+{
+	if(VRCompositor)
+	{
+		vr::TrackingUniverseOrigin NewOrigin;
+
+		switch(NewSpace)
+		{
+			case ESteamVRTrackingSpace::Seated:
+				NewOrigin = vr::TrackingUniverseOrigin::TrackingUniverseSeated;
+				break;
+			case ESteamVRTrackingSpace::Standing:
+			default:
+				NewOrigin = vr::TrackingUniverseOrigin::TrackingUniverseStanding;
+				break;
+		}
+
+		VRCompositor->SetTrackingSpace(NewOrigin);
+	}
+}
+
+ESteamVRTrackingSpace::Type FSteamVRHMD::GetTrackingSpace() const
+{
+	if(VRCompositor)
+	{
+		const vr::TrackingUniverseOrigin CurrentOrigin = VRCompositor->GetTrackingSpace();
+
+		switch(CurrentOrigin)
+		{
+		case vr::TrackingUniverseOrigin::TrackingUniverseSeated:
+			return ESteamVRTrackingSpace::Seated;
+		case vr::TrackingUniverseOrigin::TrackingUniverseStanding:
+		default:
+			return ESteamVRTrackingSpace::Standing;
+		}
+	}
+
+	// By default, assume standing
+	return ESteamVRTrackingSpace::Standing;
+}
+
 void FSteamVRHMD::PoseToOrientationAndPosition(const vr::HmdMatrix34_t& InPose, FQuat& OutOrientation, FVector& OutPosition) const
 {
 	FMatrix Pose = ToFMatrix(InPose);
@@ -277,15 +318,35 @@ void FSteamVRHMD::GetCurrentOrientationAndPosition(FQuat& CurrentOrientation, FV
 	CurrentPosition = CurHmdPosition;
 }
 
-void FSteamVRHMD::GetTrackedDeviceIds(TArray<int32>& TrackedIds)
+ESteamVRTrackedDeviceType::Type FSteamVRHMD::GetTrackedDeviceType(uint32 DeviceId) const
+{
+	vr::TrackedDeviceClass DeviceClass = VRSystem->GetTrackedDeviceClass(DeviceId);
+
+	switch (DeviceClass)
+	{
+	case vr::TrackedDeviceClass_Controller:
+		return ESteamVRTrackedDeviceType::Controller;
+	case vr::TrackedDeviceClass_TrackingReference:
+		return ESteamVRTrackedDeviceType::TrackingReference;
+	case vr::TrackedDeviceClass_Other:
+		return ESteamVRTrackedDeviceType::Other;
+	default:
+		return ESteamVRTrackedDeviceType::Invalid;
+	}
+}
+
+
+void FSteamVRHMD::GetTrackedDeviceIds(ESteamVRTrackedDeviceType::Type DeviceType, TArray<int32>& TrackedIds)
 {
 	TrackedIds.Empty();
 
 	for (uint32 i = 0; i < vr::k_unMaxTrackedDeviceCount; ++i)
 	{
 		// Add only devices with a currently valid tracked pose, and exclude the HMD
-		if ((i != vr::k_unTrackedDeviceIndex_Hmd) && TrackingFrame.bPoseIsValid[i])
-	{
+		if ((i != vr::k_unTrackedDeviceIndex_Hmd) 
+			&& TrackingFrame.bPoseIsValid[i]
+			&& (GetTrackedDeviceType(i) == DeviceType))
+		{
 			TrackedIds.Add(i);
 		}
 	}
diff --git a/Engine/Plugins/Runtime/Steam/SteamVR/Source/SteamVR/Private/SteamVRHMD.h b/Engine/Plugins/Runtime/Steam/SteamVR/Source/SteamVR/Private/SteamVRHMD.h
--- a/Engine/Plugins/Runtime/Steam/SteamVR/Source/SteamVR/Private/SteamVRHMD.h
+++ b/Engine/Plugins/Runtime/Steam/SteamVR/Source/SteamVR/Private/SteamVRHMD.h
@@ -15,6 +15,16 @@
 
 #include "SceneViewExtension.h"
 
+namespace ESteamVRTrackedDeviceType
+{
+	enum Type;
+}
+
+namespace ESteamVRTrackingSpace
+{
+	enum Type;
+}
+
 //@todo steamvr: better association between SteamController plugin and SteamVR plugin
 #ifndef MAX_STEAM_CONTROLLERS
 	#define MAX_STEAM_CONTROLLERS 8
@@ -168,7 +178,8 @@ public:
 	void ShutdownRendering();
 
 	/** Motion Controllers */
-	void GetTrackedDeviceIds(TArray<int32>& TrackedIds);
+	ESteamVRTrackedDeviceType::Type GetTrackedDeviceType(uint32 DeviceId) const;
+	void GetTrackedDeviceIds(ESteamVRTrackedDeviceType::Type DeviceType, TArray<int32>& TrackedIds);
 	bool GetTrackedObjectOrientationAndPosition(uint32 DeviceId, FQuat& CurrentOrientation, FVector& CurrentPosition);
 	bool GetTrackedDeviceIdFromControllerIndex(int32 ControllerIndex, int32& OutDeviceId);
 
@@ -185,6 +196,12 @@ public:
 	/** Get the windowed mirror mode.  @todo steamvr: thread safe flags */
 	int32 GetWindowMirrorMode() const { return WindowMirrorMode; }
 
+	/** Sets the tracking space for the returned coordinate system (e.g. standing, sitting) */
+	void SetTrackingSpace(TEnumAsByte<ESteamVRTrackingSpace::Type> NewSpace);
+
+	/** Returns the tracking space for the returned coordinate system (e.g. standing, sitting) */
+	ESteamVRTrackingSpace::Type GetTrackingSpace() const;
+
 public:
 	/** Constructor */
 	FSteamVRHMD(ISteamVRPlugin* SteamVRPlugin);
