- Updated GetTrackedDeviceIdFromControllerIndex() to work with the new SteamVRController plugin.

diff --git a/Engine/Plugins/Runtime/Steam/SteamVR/Source/SteamVR/Private/SteamVRHMD.cpp b/Engine/Plugins/Runtime/Steam/SteamVR/Source/SteamVR/Private/SteamVRHMD.cpp
--- a/Engine/Plugins/Runtime/Steam/SteamVR/Source/SteamVR/Private/SteamVRHMD.cpp
+++ b/Engine/Plugins/Runtime/Steam/SteamVR/Source/SteamVR/Private/SteamVRHMD.cpp
@@ -34,6 +34,19 @@ public:
 		VRSystem = InVRSystem;
 	}
 
+	virtual void SetControllerToDeviceMap(int32* InControllerToDeviceMap) override
+	{
+		if (!GEngine->HMDDevice.IsValid() || (GEngine->HMDDevice->GetHMDDeviceType() != EHMDDeviceType::DT_SteamVR))
+		{
+			// no valid SteamVR HMD found
+			return;
+		}
+
+		FSteamVRHMD* SteamVRHMD = static_cast<FSteamVRHMD*>(GEngine->HMDDevice.Get());
+
+		SteamVRHMD->SetControllerToDeviceMap(InControllerToDeviceMap);
+	}
+
 private:
 	vr::IVRSystem* VRSystem;
 };
@@ -292,6 +305,13 @@ ESteamVRTrackingSpace::Type FSteamVRHMD::GetTrackingSpace() const
 	return ESteamVRTrackingSpace::Standing;
 }
 
+void FSteamVRHMD::SetControllerToDeviceMap(int32* InControllerToDeviceMap)
+{
+	check(sizeof(ControllerToDeviceMap) == (MAX_STEAMVR_CONTROLLERS*sizeof(int32)));
+
+	FMemory::Memcpy(ControllerToDeviceMap, InControllerToDeviceMap, sizeof(ControllerToDeviceMap));
+}
+
 void FSteamVRHMD::PoseToOrientationAndPosition(const vr::HmdMatrix34_t& InPose, FQuat& OutOrientation, FVector& OutPosition) const
 {
 	FMatrix Pose = ToFMatrix(InPose);
@@ -375,12 +395,12 @@ bool FSteamVRHMD::GetTrackedDeviceIdFromControllerIndex(int32 ControllerIndex, i
 
 	OutDeviceId = -1;
 
-	if ((ControllerIndex < 0) || (ControllerIndex >= MAX_STEAM_CONTROLLERS))
+	if ((ControllerIndex < 0) || (ControllerIndex >= MAX_STEAMVR_CONTROLLERS))
 	{
 		return false;
 	}
 
-	OutDeviceId = ControllerIndex;		//@todo steamvr: remove entirely now that these always match
+	OutDeviceId = ControllerToDeviceMap[ControllerIndex];
 	
 	return (OutDeviceId != -1);
 }
@@ -938,6 +958,12 @@ void FSteamVRHMD::Startup()
 			UE_LOG(LogHMD, Warning, TEXT("Failed to initialize Chaperone.  Error: %d"), (int32)ChaperoneErr);
 		}
 
+		// Initialize our controller to device index
+		for (int32 i = 0; i < MAX_STEAMVR_CONTROLLERS; ++i)
+		{
+			ControllerToDeviceMap[i] = -1;
+		}
+
 #if PLATFORM_WINDOWS
 		if (IsPCPlatform(GMaxRHIShaderPlatform) && !IsOpenGLPlatform(GMaxRHIShaderPlatform))
 		{
diff --git a/Engine/Plugins/Runtime/Steam/SteamVR/Source/SteamVR/Private/SteamVRHMD.h b/Engine/Plugins/Runtime/Steam/SteamVR/Source/SteamVR/Private/SteamVRHMD.h
--- a/Engine/Plugins/Runtime/Steam/SteamVR/Source/SteamVR/Private/SteamVRHMD.h
+++ b/Engine/Plugins/Runtime/Steam/SteamVR/Source/SteamVR/Private/SteamVRHMD.h
@@ -25,9 +25,9 @@ namespace ESteamVRTrackingSpace
 	enum Type;
 }
 
-//@todo steamvr: better association between SteamController plugin and SteamVR plugin
-#ifndef MAX_STEAM_CONTROLLERS
-	#define MAX_STEAM_CONTROLLERS 8
+//@todo steamvr: better association between SteamVRController plugin and SteamVR plugin
+#ifndef MAX_STEAMVR_CONTROLLERS
+	#define MAX_STEAMVR_CONTROLLERS 8
 #endif
 
 /** Stores vectors, in clockwise order, to define soft and hard bounds for Chaperone */
@@ -202,6 +202,9 @@ public:
 	/** Returns the tracking space for the returned coordinate system (e.g. standing, sitting) */
 	ESteamVRTrackingSpace::Type GetTrackingSpace() const;
 
+	/** Sets the map from controller index to tracked device index. */
+	void SetControllerToDeviceMap(int32* InControllerToDeviceMap);
+
 public:
 	/** Constructor */
 	FSteamVRHMD(ISteamVRPlugin* SteamVRPlugin);
@@ -373,6 +376,9 @@ private:
 	/** World units (UU) to Meters scale.  Read from the level, and used to transform positional tracking data */
 	float WorldToMetersScale;
 
+	/** Mapping from Controller index to tracked device index.  Passed in from the controller plugin */
+	int32 ControllerToDeviceMap[MAX_STEAMVR_CONTROLLERS];
+
 	IRendererModule* RendererModule;
 	ISteamVRPlugin* SteamVRPlugin;
 
diff --git a/Engine/Plugins/Runtime/Steam/SteamVR/Source/SteamVR/Public/ISteamVRPlugin.h b/Engine/Plugins/Runtime/Steam/SteamVR/Source/SteamVR/Public/ISteamVRPlugin.h
--- a/Engine/Plugins/Runtime/Steam/SteamVR/Source/SteamVR/Public/ISteamVRPlugin.h
+++ b/Engine/Plugins/Runtime/Steam/SteamVR/Source/SteamVR/Public/ISteamVRPlugin.h
@@ -38,9 +38,25 @@ public:
 		return FModuleManager::Get().IsModuleLoaded( "SteamVR" );
 	}
 
+	/**
+	 * Get the IVRSystem* that was previously set by the HMD implemenentation.
+	 *
+	 * @return The pointer if the HMD has been initialized, otherwise nullptr 
+	 */
 	virtual vr::IVRSystem* GetVRSystem() const=0;
+
+	/**
+	 * Set the cached IVRSystem pointer so that the controller can retrieve it.
+	 * The HMD passes this to the controller.
+	 */
 	virtual void SetVRSystem(vr::IVRSystem* VRSystem)=0;
 
+	/**
+	 * Update the Controller to Device mapping.
+	 * The controller passes this to the HMD.
+	 */
+	virtual void SetControllerToDeviceMap(int32* InControllerToDeviceMap)=0;
+
 private:
 };
 
diff --git a/Engine/Plugins/Runtime/Steam/SteamVR/Source/SteamVRController/Private/SteamVRController.cpp b/Engine/Plugins/Runtime/Steam/SteamVR/Source/SteamVRController/Private/SteamVRController.cpp
--- a/Engine/Plugins/Runtime/Steam/SteamVR/Source/SteamVRController/Private/SteamVRController.cpp
+++ b/Engine/Plugins/Runtime/Steam/SteamVR/Source/SteamVRController/Private/SteamVRController.cpp
@@ -100,6 +100,9 @@ public:
 					DeviceToControllerMap[DeviceIndex] = NumControllersMapped;
 					ControllerToDeviceMap[NumControllersMapped] = DeviceIndex;
 					++NumControllersMapped;
+
+					// update the SteamVR plugin with the new mapping
+					SteamVRPlugin->SetControllerToDeviceMap(ControllerToDeviceMap);
 				}
 
 				// get the controller index for this device
