- Fixed crash if SteamVRController inits but SteamVRHMD does not due to Oculus plugin initing first.

diff --git a/Engine/Plugins/Runtime/Steam/SteamVR/Source/SteamVR/Private/SteamVRHMD.cpp b/Engine/Plugins/Runtime/Steam/SteamVR/Source/SteamVR/Private/SteamVRHMD.cpp
--- a/Engine/Plugins/Runtime/Steam/SteamVR/Source/SteamVR/Private/SteamVRHMD.cpp
+++ b/Engine/Plugins/Runtime/Steam/SteamVR/Source/SteamVR/Private/SteamVRHMD.cpp
@@ -24,6 +24,11 @@ class FSteamVRPlugin : public ISteamVRPlugin
 	}
 
 public:
+	FSteamVRPlugin::FSteamVRPlugin()
+		: VRSystem(nullptr)
+	{
+	}
+
 	virtual vr::IVRSystem* GetVRSystem() const override
 	{
 		return VRSystem;
