- Disable finishcurrentframe on SteamVR

diff --git a/Engine/Plugins/Runtime/Steam/SteamVR/Source/SteamVR/Private/SteamVRHMD.cpp b/Engine/Plugins/Runtime/Steam/SteamVR/Source/SteamVR/Private/SteamVRHMD.cpp
--- a/Engine/Plugins/Runtime/Steam/SteamVR/Source/SteamVR/Private/SteamVRHMD.cpp
+++ b/Engine/Plugins/Runtime/Steam/SteamVR/Source/SteamVR/Private/SteamVRHMD.cpp
@@ -924,7 +924,7 @@ void FSteamVRHMD::Startup()
 
 		// enforce finishcurrentframe
 		static IConsoleVariable* CFCFVar = IConsoleManager::Get().FindConsoleVariable(TEXT("r.finishcurrentframe"));
-		CFCFVar->Set(true);
+		CFCFVar->Set(false);
 
 		// Uncap fps to enable FPS higher than 62
 		GEngine->bSmoothFrameRate = false;
