#ue4 Fix for build breaks

diff --git a/Engine/Plugins/Runtime/GearVR/Source/GearVR/Private/GearVR.cpp b/Engine/Plugins/Runtime/GearVR/Source/GearVR/Private/GearVR.cpp
--- a/Engine/Plugins/Runtime/GearVR/Source/GearVR/Private/GearVR.cpp
+++ b/Engine/Plugins/Runtime/GearVR/Source/GearVR/Private/GearVR.cpp
@@ -952,7 +952,7 @@ void FViewExtension::PreRenderViewFamily_RenderThread(FRHICommandListImmediate&
 	}
 }
 
-void FGearVR::CalculateRenderTargetSize(const FViewport& Viewport, uint32& InOutSizeX, uint32& InOutSizeY) const
+void FGearVR::CalculateRenderTargetSize(const FViewport& Viewport, uint32& InOutSizeX, uint32& InOutSizeY)
 {
 	check(IsInGameThread());
 
@@ -968,7 +968,7 @@ void FGearVR::GetOrthoProjection(int32 RTWidth, int32 RTHeight, float OrthoDista
 	OrthoProjection[1] = FTranslationMatrix(FVector(OrthoProjection[1].M[0][3] * RTWidth * .25 + RTWidth * .5, 0 , 0));
 }
 
-bool FGearVR::NeedReAllocateViewportRenderTarget(const FViewport& Viewport) const
+bool FGearVR::NeedReAllocateViewportRenderTarget(const FViewport& Viewport)
 {
 	check(IsInGameThread());
 
@@ -1263,7 +1263,7 @@ void FGearVRCustomPresent::UpdateViewport(const FViewport& Viewport, FRHIViewpor
 
 
 
-bool FGearVRCustomPresent::Present(int SyncInterval)
+bool FGearVRCustomPresent::Present(int& SyncInterval)
 {
 	check(IsInRenderingThread());
 
diff --git a/Engine/Plugins/Runtime/GearVR/Source/GearVR/Private/GearVR.h b/Engine/Plugins/Runtime/GearVR/Source/GearVR/Private/GearVR.h
--- a/Engine/Plugins/Runtime/GearVR/Source/GearVR/Private/GearVR.h
+++ b/Engine/Plugins/Runtime/GearVR/Source/GearVR/Private/GearVR.h
@@ -202,7 +202,7 @@ public:
 	// Resets Viewport-specific resources.
 	virtual void OnBackBufferResize() override;
 	// Returns true if Engine should perform its own Present.
-	virtual bool Present(int SyncInterval) override;
+	virtual bool Present(int& SyncInterval) override;
 
 private:
 	void DicedBlit(uint32 SourceX, uint32 SourceY, uint32 DestX, uint32 DestY, uint32 Width, uint32 Height, uint32 NumXSteps, uint32 NumYSteps);
@@ -254,8 +254,8 @@ public:
 	virtual void InitCanvasFromView(FSceneView* InView, UCanvas* Canvas) override;
 
 	virtual void UpdateViewport(bool bUseSeparateRenderTarget, const FViewport& Viewport, SViewport* ViewportWidget) override;
-	virtual void CalculateRenderTargetSize(const class FViewport& Viewport, uint32& InOutSizeX, uint32& InOutSizeY) const override;
-	virtual bool NeedReAllocateViewportRenderTarget(const FViewport& Viewport) const override;
+	virtual void CalculateRenderTargetSize(const class FViewport& Viewport, uint32& InOutSizeX, uint32& InOutSizeY) override;
+	virtual bool NeedReAllocateViewportRenderTarget(const FViewport& Viewport) override;
 
 	virtual bool ShouldUseSeparateRenderTarget() const override
 	{
diff --git a/Engine/Plugins/Runtime/Steam/SteamVR/Source/SteamVR/Private/SteamVRHMD.cpp b/Engine/Plugins/Runtime/Steam/SteamVR/Source/SteamVR/Private/SteamVRHMD.cpp
--- a/Engine/Plugins/Runtime/Steam/SteamVR/Source/SteamVR/Private/SteamVRHMD.cpp
+++ b/Engine/Plugins/Runtime/Steam/SteamVR/Source/SteamVR/Private/SteamVRHMD.cpp
@@ -767,7 +767,7 @@ FSteamVRHMD::BridgeBaseImpl* FSteamVRHMD::GetActiveRHIBridgeImpl()
 	return nullptr;
 }
 
-void FSteamVRHMD::CalculateRenderTargetSize(const class FViewport& Viewport, uint32& InOutSizeX, uint32& InOutSizeY) const
+void FSteamVRHMD::CalculateRenderTargetSize(const class FViewport& Viewport, uint32& InOutSizeX, uint32& InOutSizeY)
 {
 	check(IsInGameThread());
 
@@ -783,7 +783,7 @@ void FSteamVRHMD::CalculateRenderTargetSize(const class FViewport& Viewport, uin
 	}
 }
 
-bool FSteamVRHMD::NeedReAllocateViewportRenderTarget(const FViewport& Viewport) const
+bool FSteamVRHMD::NeedReAllocateViewportRenderTarget(const FViewport& Viewport)
 {
 	check(IsInGameThread());
 
diff --git a/Engine/Plugins/Runtime/Steam/SteamVR/Source/SteamVR/Private/SteamVRHMD.h b/Engine/Plugins/Runtime/Steam/SteamVR/Source/SteamVR/Private/SteamVRHMD.h
--- a/Engine/Plugins/Runtime/Steam/SteamVR/Source/SteamVR/Private/SteamVRHMD.h
+++ b/Engine/Plugins/Runtime/Steam/SteamVR/Source/SteamVR/Private/SteamVRHMD.h
@@ -105,8 +105,8 @@ public:
 	virtual void RenderTexture_RenderThread(FRHICommandListImmediate& RHICmdList, FTexture2DRHIParamRef BackBuffer, FTexture2DRHIParamRef SrcTexture) const override;
 	virtual void GetEyeRenderParams_RenderThread(const FRenderingCompositePassContext& Context, FVector2D& EyeToSrcUVScaleValue, FVector2D& EyeToSrcUVOffsetValue) const override;
 
-	virtual void CalculateRenderTargetSize(const class FViewport& Viewport, uint32& InOutSizeX, uint32& InOutSizeY) const override;
-	virtual bool NeedReAllocateViewportRenderTarget(const FViewport& Viewport) const override;
+	virtual void CalculateRenderTargetSize(const class FViewport& Viewport, uint32& InOutSizeX, uint32& InOutSizeY) override;
+	virtual bool NeedReAllocateViewportRenderTarget(const FViewport& Viewport) override;
 	virtual bool ShouldUseSeparateRenderTarget() const override
 	{
 		check(IsInGameThread());
@@ -154,7 +154,7 @@ public:
 		D3D11Bridge(FSteamVRHMD* plugin);
 
 		virtual void OnBackBufferResize() override;
-		virtual bool Present(int SyncInterval) override;
+		virtual bool Present(int& SyncInterval) override;
 
 		virtual void BeginRendering() override;
 		void FinishRendering();
diff --git a/Engine/Plugins/Runtime/Steam/SteamVR/Source/SteamVR/Private/SteamVRRender.cpp b/Engine/Plugins/Runtime/Steam/SteamVR/Source/SteamVR/Private/SteamVRRender.cpp
--- a/Engine/Plugins/Runtime/Steam/SteamVR/Source/SteamVR/Private/SteamVRRender.cpp
+++ b/Engine/Plugins/Runtime/Steam/SteamVR/Source/SteamVR/Private/SteamVRRender.cpp
@@ -142,7 +142,7 @@ void FSteamVRHMD::D3D11Bridge::OnBackBufferResize()
 {
 }
 
-bool FSteamVRHMD::D3D11Bridge::Present(int SyncInterval)
+bool FSteamVRHMD::D3D11Bridge::Present(int& SyncInterval)
 {
 	check(IsInRenderingThread());
 
diff --git a/Engine/Plugins/Runtime/Steam/SteamVR/SteamVR.uplugin b/Engine/Plugins/Runtime/Steam/SteamVR/SteamVR.uplugin
--- a/Engine/Plugins/Runtime/Steam/SteamVR/SteamVR.uplugin
+++ b/Engine/Plugins/Runtime/Steam/SteamVR/SteamVR.uplugin
@@ -9,7 +9,7 @@
 	"EngineVersion" : "4.2.0",
 	"Description" : "Support for the SteamVR API",
 	"Category" : "Virtual Reality",
-	"EnabledByDefault" : false,
+	"EnabledByDefault" : true,
 
 	"Modules" :
 	[
