Temporarily disabling SteamVR plugin until we figure out a dll dependency issue.

diff --git a/Engine/Plugins/Runtime/Steam/SteamVR/SteamVR.uplugin b/Engine/Plugins/Runtime/Steam/SteamVR/SteamVR.uplugin
--- a/Engine/Plugins/Runtime/Steam/SteamVR/SteamVR.uplugin
+++ b/Engine/Plugins/Runtime/Steam/SteamVR/SteamVR.uplugin
@@ -9,7 +9,7 @@
 	"EngineVersion" : "4.2.0",
 	"Description" : "Support for the SteamVR API",
 	"Category" : "Virtual Reality",
-	"EnabledByDefault" : true,
+	"EnabledByDefault" : false,
 
 	"Modules" :
 	[
