Changed SteamVR init failure to Log rather than Warning since it will be a common case.

diff --git a/Engine/Plugins/Runtime/Steam/SteamVR/Source/SteamVR/Private/SteamVRHMD.cpp b/Engine/Plugins/Runtime/Steam/SteamVR/Source/SteamVR/Private/SteamVRHMD.cpp
--- a/Engine/Plugins/Runtime/Steam/SteamVR/Source/SteamVR/Private/SteamVRHMD.cpp
+++ b/Engine/Plugins/Runtime/Steam/SteamVR/Source/SteamVR/Private/SteamVRHMD.cpp
@@ -967,7 +967,7 @@ void FSteamVRHMD::Startup()
 	}
 	else
 	{
-		UE_LOG(LogHMD, Warning, TEXT("SteamVR failed to initialize.  Error: %d"), (int32)HmdErr);
+		UE_LOG(LogHMD, Log, TEXT("SteamVR failed to initialize.  Err: %d"), (int32)HmdErr);
 
 		VRSystem = nullptr;
 	}
