#steamvr Fixes for exporting the SteamVR function library for use in other modules, fix for faulty controller input in SteamController

diff --git a/Engine/Plugins/Runtime/Steam/SteamVR/Source/SteamVR/Classes/SteamVRFunctionLibrary.h b/Engine/Plugins/Runtime/Steam/SteamVR/Source/SteamVR/Classes/SteamVRFunctionLibrary.h
--- a/Engine/Plugins/Runtime/Steam/SteamVR/Source/SteamVR/Classes/SteamVRFunctionLibrary.h
+++ b/Engine/Plugins/Runtime/Steam/SteamVR/Source/SteamVR/Classes/SteamVRFunctionLibrary.h
@@ -40,10 +40,12 @@ namespace ESteamVRTrackingSpace
 /**
  * SteamVR Extensions Function Library
  */
-UCLASS(MinimalAPI)
-class USteamVRFunctionLibrary : public UBlueprintFunctionLibrary
+UCLASS()
+class STEAMVR_API USteamVRFunctionLibrary : public UBlueprintFunctionLibrary
 {
 	GENERATED_UCLASS_BODY()
+
+public:
 	
 	/**
 	 * Returns an array of the currently tracked device IDs
diff --git a/Engine/Plugins/Runtime/Steam/SteamVR/Source/SteamVRController/Private/SteamVRController.cpp b/Engine/Plugins/Runtime/Steam/SteamVR/Source/SteamVRController/Private/SteamVRController.cpp
--- a/Engine/Plugins/Runtime/Steam/SteamVR/Source/SteamVRController/Private/SteamVRController.cpp
+++ b/Engine/Plugins/Runtime/Steam/SteamVR/Source/SteamVRController/Private/SteamVRController.cpp
@@ -113,13 +113,25 @@ public:
 					if (ControllerState.unPacketNum != ControllerStates[ControllerIndex].PacketNum )
 					{
 						bool CurrentStates[MAX_NUM_CONTROLLER_BUTTONS] = {0};
-			
+
 						// Get the current state of all buttons
 						CurrentStates[0] = !!(ControllerState.ulButtonPressed & vr::ButtonMaskFromId(vr::k_EButton_System));
 						CurrentStates[1] = !!(ControllerState.ulButtonPressed & vr::ButtonMaskFromId(vr::k_EButton_ApplicationMenu));
 						CurrentStates[2] = !!(ControllerState.ulButtonPressed & vr::ButtonMaskFromId(vr::k_EButton_SteamVR_Touchpad));
 						CurrentStates[3] = !!(ControllerState.ulButtonPressed & vr::ButtonMaskFromId(vr::k_EButton_SteamVR_Trigger));
 						CurrentStates[4] = !!(ControllerState.ulButtonPressed & vr::ButtonMaskFromId(vr::k_EButton_Grip));
+
+						const bool bTouchPadTouched = (ControllerState.ulButtonTouched & vr::ButtonMaskFromId(vr::k_EButton_SteamVR_Touchpad)) == vr::ButtonMaskFromId(vr::k_EButton_SteamVR_Touchpad);
+						const bool bTouchPadPressed = (ControllerState.ulButtonPressed & vr::ButtonMaskFromId(vr::k_EButton_SteamVR_Touchpad)) == vr::ButtonMaskFromId(vr::k_EButton_SteamVR_Touchpad);
+
+						// If the touchpad isn't currently pressed, zero put both of the axes
+						if (!bTouchPadTouched || 
+							!bTouchPadPressed)	// @todo: Hack to workaround broken hardware where Touched event is not reported correctly
+						{
+							ControllerState.rAxis[TOUCHPAD_AXIS].y = 0.0f;
+							ControllerState.rAxis[TOUCHPAD_AXIS].x = 0.0f;
+						}
+
 						CurrentStates[5] = !!(ControllerState.rAxis[TOUCHPAD_AXIS].y > TOUCHPAD_DEADZONE);
 						CurrentStates[6] = !!(ControllerState.rAxis[TOUCHPAD_AXIS].y < -TOUCHPAD_DEADZONE);
 						CurrentStates[7] = !!(ControllerState.rAxis[TOUCHPAD_AXIS].x < -TOUCHPAD_DEADZONE);
