#vr Merging in improvements to the SteamVRController implementation, in preparation for rewrite of the motion control path

diff --git a/Engine/Plugins/Runtime/Steam/SteamVR/Source/SteamVR/Classes/SteamVRFunctionLibrary.h b/Engine/Plugins/Runtime/Steam/SteamVR/Source/SteamVR/Classes/SteamVRFunctionLibrary.h
--- a/Engine/Plugins/Runtime/Steam/SteamVR/Source/SteamVR/Classes/SteamVRFunctionLibrary.h
+++ b/Engine/Plugins/Runtime/Steam/SteamVR/Source/SteamVR/Classes/SteamVRFunctionLibrary.h
@@ -3,39 +3,44 @@
 #pragma once
 #include "SteamVRFunctionLibrary.generated.h"
 
-/** Describes the type of device tracked in the SteamVR volume (e.g. controller, Lighthouse, etc) */
+/** Defines the class of tracked devices in SteamVR*/
 UENUM(BlueprintType)
-namespace ESteamVRTrackedDeviceType
+enum class ESteamVRTrackedDeviceType
 {
-	enum Type
-	{
-		/** Represents a Steam VR Controller */
-		Controller,
+	/** Represents a Steam VR Controller */
+	Controller,
 
-		/** Represents a static tracking reference device, such as a Lighthouse or tracking camera */
-		TrackingReference,
+	/** Represents a static tracking reference device, such as a Lighthouse or tracking camera */
+	TrackingReference,
 
-		/** Misc. device types, for future expansion */
-		Other,
+	/** Misc. device types, for future expansion */
+	Other,
 
-		/** DeviceId is invalid */
-		Invalid
-	};
-}
+	/** DeviceId is invalid */
+	Invalid
+};
 
 /** Defines which calibration point to use (e.g. center of the room for standing, calibrated head position for seated) */
 UENUM(BlueprintType)
-namespace ESteamVRTrackingSpace
+enum class ESteamVRTrackingSpace
+{
+	/** Standing origin, where poses are relative to the safe room bounds provided by the user's calibration */
+	Standing,
+
+	/** Seated origin, where poses are relative to the user's calibrated seated head position */
+	Seated
+};
+
+/** The motion controller hand for a specific Unreal controller.  Each Unreal controller can support a pair of devices, one for the left hand and one for the right hand. */
+UENUM(BlueprintType)
+enum class ESteamVRControllerHand
 {
-	enum Type
-	{
-		/** Standing origin, where poses are relative to the safe room bounds provided by the user's calibration */
-		Standing,
+	/** Left hand */
+	Left,
 
-		/** Seated origin, where poses are relative to the user's calibrated seated head position */
-		Seated
-	};
-}
+	/** Right hand */
+	Right,
+};
 
 /**
  * SteamVR Extensions Function Library
@@ -54,7 +59,7 @@ public:
 	 * @param	OutTrackedDeviceIds	(out) Array containing the ID of each device that's currently tracked
 	 */
 	UFUNCTION(BlueprintPure, Category="SteamVR")
-	static void GetValidTrackedDeviceIds(TEnumAsByte<ESteamVRTrackedDeviceType::Type> DeviceType, TArray<int32>& OutTrackedDeviceIds);
+	static void GetValidTrackedDeviceIds(TEnumAsByte<ESteamVRTrackedDeviceType> DeviceType, TArray<int32>& OutTrackedDeviceIds);
 
 	/**
 	 * Gets the orientation and position (in device space) of the device with the specified ID
@@ -68,14 +73,16 @@ public:
 	static bool GetTrackedDevicePositionAndOrientation(int32 DeviceId, FVector& OutPosition, FRotator& OutOrientation);
 
 	/**
-	 * Given a controller index, returns the attached tracked device ID
+	 * Given a controller index and a hand, returns the position and orientation of the controller
 	 *
 	 * @param	ControllerIndex	Index of the controller to get the tracked device ID for
-	 * @param	OutDeviceId		(out) Tracked device ID for the controller
+	 * @param	Hand			Which hand's controller to get the position and orientation for
+	 * @param	OutPosition		(out) Current position of the device
+	 * @param	OutRotation		(out) Current rotation of the device
 	 * @return	True if the specified controller index has a valid tracked device ID
 	 */
 	UFUNCTION(BlueprintPure, Category = "SteamVR")
-	static bool GetTrackedDeviceIdFromControllerIndex(int32 ControllerIndex, int32& OutDeviceId);
+	static bool GetHandPositionAndOrientation(int32 ControllerIndex, ESteamVRControllerHand Hand, FVector& OutPosition, FRotator& OutOrientation);
 
 	/**
 	 * Sets the tracking space (e.g. sitting or standing), which changes which space tracked positions are returned to
@@ -83,11 +90,11 @@ public:
 	 * NewSpace		The new space to consider all tracked positions in.  For instance, standing assumes the center of the safe zone as the origin.
 	 */
 	UFUNCTION(BlueprintCallable, Category = "SteamVR")
-	static void SetTrackingSpace(TEnumAsByte<ESteamVRTrackingSpace::Type> NewSpace);
+	static void SetTrackingSpace(TEnumAsByte<ESteamVRTrackingSpace> NewSpace);
 
 	/**
 	 * Gets the tracking space (e.g. sitting or standing), which determines the location of the origin.
 	 */
 	UFUNCTION(BlueprintPure, Category = "SteamVR")
-	static TEnumAsByte<ESteamVRTrackingSpace::Type> GetTrackingSpace();
+	static TEnumAsByte<ESteamVRTrackingSpace> GetTrackingSpace();
 };
diff --git a/Engine/Plugins/Runtime/Steam/SteamVR/Source/SteamVR/Private/SteamVRFunctionLibrary.cpp b/Engine/Plugins/Runtime/Steam/SteamVR/Source/SteamVR/Private/SteamVRFunctionLibrary.cpp
--- a/Engine/Plugins/Runtime/Steam/SteamVR/Source/SteamVR/Private/SteamVRFunctionLibrary.cpp
+++ b/Engine/Plugins/Runtime/Steam/SteamVR/Source/SteamVR/Private/SteamVRFunctionLibrary.cpp
@@ -19,7 +19,7 @@ FSteamVRHMD* GetSteamVRHMD()
 	return nullptr;
 }
 
-void USteamVRFunctionLibrary::GetValidTrackedDeviceIds(TEnumAsByte<ESteamVRTrackedDeviceType::Type> DeviceType, TArray<int32>& OutTrackedDeviceIds)
+void USteamVRFunctionLibrary::GetValidTrackedDeviceIds(TEnumAsByte<ESteamVRTrackedDeviceType> DeviceType, TArray<int32>& OutTrackedDeviceIds)
 {
 	OutTrackedDeviceIds.Empty();
 
@@ -45,20 +45,22 @@ bool USteamVRFunctionLibrary::GetTrackedDevicePositionAndOrientation(int32 Devic
 	return RetVal;
 }
 
-bool USteamVRFunctionLibrary::GetTrackedDeviceIdFromControllerIndex(int32 ControllerIndex, int32& OutDeviceId)
+bool USteamVRFunctionLibrary::GetHandPositionAndOrientation(int32 ControllerIndex, ESteamVRControllerHand Hand, FVector& OutPosition, FRotator& OutOrientation)
 {
-	OutDeviceId = -1;
+	bool RetVal = false;
 
 	FSteamVRHMD* SteamVRHMD = GetSteamVRHMD();
 	if (SteamVRHMD)
 	{
-		return SteamVRHMD->GetTrackedDeviceIdFromControllerIndex(ControllerIndex, OutDeviceId);
+		FQuat DeviceOrientation = FQuat::Identity;
+		RetVal = SteamVRHMD->GetControllerHandPositionAndOrientation(ControllerIndex, Hand, OutPosition, DeviceOrientation);
+		OutOrientation = DeviceOrientation.Rotator();
 	}
 
 	return false;
 }
 
-void USteamVRFunctionLibrary::SetTrackingSpace(TEnumAsByte<ESteamVRTrackingSpace::Type> NewSpace)
+void USteamVRFunctionLibrary::SetTrackingSpace(TEnumAsByte<ESteamVRTrackingSpace> NewSpace)
 {
 	FSteamVRHMD* SteamVRHMD = GetSteamVRHMD();
 	if (SteamVRHMD)
@@ -67,9 +69,9 @@ void USteamVRFunctionLibrary::SetTrackingSpace(TEnumAsByte<ESteamVRTrackingSpace
 	}
 }
 
-TEnumAsByte<ESteamVRTrackingSpace::Type> USteamVRFunctionLibrary::GetTrackingSpace()
+TEnumAsByte<ESteamVRTrackingSpace> USteamVRFunctionLibrary::GetTrackingSpace()
 {
-	ESteamVRTrackingSpace::Type RetVal = ESteamVRTrackingSpace::Standing;
+	ESteamVRTrackingSpace RetVal = ESteamVRTrackingSpace::Standing;
 
 	FSteamVRHMD* SteamVRHMD = GetSteamVRHMD();
 	if (SteamVRHMD)
diff --git a/Engine/Plugins/Runtime/Steam/SteamVR/Source/SteamVR/Private/SteamVRHMD.cpp b/Engine/Plugins/Runtime/Steam/SteamVR/Source/SteamVR/Private/SteamVRHMD.cpp
--- a/Engine/Plugins/Runtime/Steam/SteamVR/Source/SteamVR/Private/SteamVRHMD.cpp
+++ b/Engine/Plugins/Runtime/Steam/SteamVR/Source/SteamVR/Private/SteamVRHMD.cpp
@@ -39,7 +39,7 @@ public:
 		VRSystem = InVRSystem;
 	}
 
-	virtual void SetControllerToDeviceMap(int32* InControllerToDeviceMap) override
+	virtual void SetUnrealControllerIdAndHandToDeviceIdMap(int32 InUnrealControllerIdAndHandToDeviceIdMap[MAX_STEAMVR_CONTROLLER_PAIRS][2]) override
 	{
 		if (!GEngine->HMDDevice.IsValid() || (GEngine->HMDDevice->GetHMDDeviceType() != EHMDDeviceType::DT_SteamVR))
 		{
@@ -49,7 +49,7 @@ public:
 
 		FSteamVRHMD* SteamVRHMD = static_cast<FSteamVRHMD*>(GEngine->HMDDevice.Get());
 
-		SteamVRHMD->SetControllerToDeviceMap(InControllerToDeviceMap);
+		SteamVRHMD->SetUnrealControllerIdAndHandToDeviceIdMap(InUnrealControllerIdAndHandToDeviceIdMap);
 	}
 
 private:
@@ -161,7 +161,7 @@ void FSteamVRHMD::GetCurrentPose(FQuat& CurrentOrientation, FVector& CurrentPosi
 		return;
 	}
 
-	check(DeviceId < vr::k_unMaxTrackedDeviceCount);
+	check(DeviceId >= 0 && DeviceId < vr::k_unMaxTrackedDeviceCount);
 
 	if (bForceRefresh)
 	{
@@ -269,7 +269,7 @@ TArray<FVector> FSteamVRHMD::GetHardBounds() const
 	return Bounds;
 }
 
-void FSteamVRHMD::SetTrackingSpace(TEnumAsByte<ESteamVRTrackingSpace::Type> NewSpace)
+void FSteamVRHMD::SetTrackingSpace(TEnumAsByte<ESteamVRTrackingSpace> NewSpace)
 {
 	if(VRCompositor)
 	{
@@ -290,7 +290,7 @@ void FSteamVRHMD::SetTrackingSpace(TEnumAsByte<ESteamVRTrackingSpace::Type> NewS
 	}
 }
 
-ESteamVRTrackingSpace::Type FSteamVRHMD::GetTrackingSpace() const
+ESteamVRTrackingSpace FSteamVRHMD::GetTrackingSpace() const
 {
 	if(VRCompositor)
 	{
@@ -310,11 +310,15 @@ ESteamVRTrackingSpace::Type FSteamVRHMD::GetTrackingSpace() const
 	return ESteamVRTrackingSpace::Standing;
 }
 
-void FSteamVRHMD::SetControllerToDeviceMap(int32* InControllerToDeviceMap)
+void FSteamVRHMD::SetUnrealControllerIdAndHandToDeviceIdMap(int32 InUnrealControllerIdAndHandToDeviceIdMap[ MAX_STEAMVR_CONTROLLER_PAIRS ][ 2 ] )
 {
-	check(sizeof(ControllerToDeviceMap) == (MAX_STEAMVR_CONTROLLERS*sizeof(int32)));
-
-	FMemory::Memcpy(ControllerToDeviceMap, InControllerToDeviceMap, sizeof(ControllerToDeviceMap));
+	for( int32 UnrealControllerIndex = 0; UnrealControllerIndex < MAX_STEAMVR_CONTROLLER_PAIRS; ++UnrealControllerIndex )
+	{
+		for( int32 HandIndex = 0; HandIndex < 2; ++HandIndex )
+		{
+			UnrealControllerIdAndHandToDeviceIdMap[ UnrealControllerIndex ][ HandIndex ] = InUnrealControllerIdAndHandToDeviceIdMap[ UnrealControllerIndex ][ HandIndex ];
+		}
+	}
 }
 
 void FSteamVRHMD::PoseToOrientationAndPosition(const vr::HmdMatrix34_t& InPose, FQuat& OutOrientation, FVector& OutPosition) const
@@ -343,7 +347,7 @@ void FSteamVRHMD::GetCurrentOrientationAndPosition(FQuat& CurrentOrientation, FV
 	CurrentPosition = CurHmdPosition;
 }
 
-ESteamVRTrackedDeviceType::Type FSteamVRHMD::GetTrackedDeviceType(uint32 DeviceId) const
+ESteamVRTrackedDeviceType FSteamVRHMD::GetTrackedDeviceType(uint32 DeviceId) const
 {
 	vr::TrackedDeviceClass DeviceClass = VRSystem->GetTrackedDeviceClass(DeviceId);
 
@@ -361,7 +365,7 @@ ESteamVRTrackedDeviceType::Type FSteamVRHMD::GetTrackedDeviceType(uint32 DeviceI
 }
 
 
-void FSteamVRHMD::GetTrackedDeviceIds(ESteamVRTrackedDeviceType::Type DeviceType, TArray<int32>& TrackedIds)
+void FSteamVRHMD::GetTrackedDeviceIds(ESteamVRTrackedDeviceType DeviceType, TArray<int32>& TrackedIds)
 {
 	TrackedIds.Empty();
 
@@ -383,7 +387,7 @@ bool FSteamVRHMD::GetTrackedObjectOrientationAndPosition(uint32 DeviceId, FQuat&
 
 	bool bHasValidPose = false;
 
-	if (DeviceId < vr::k_unMaxTrackedDeviceCount)
+	if (DeviceId >= 0 && DeviceId < vr::k_unMaxTrackedDeviceCount)
 	{
 		CurrentOrientation = TrackingFrame.DeviceOrientation[DeviceId];
 		CurrentPosition = TrackingFrame.DevicePosition[DeviceId];
@@ -394,22 +398,21 @@ bool FSteamVRHMD::GetTrackedObjectOrientationAndPosition(uint32 DeviceId, FQuat&
 	return bHasValidPose;
 }
 
-bool FSteamVRHMD::GetTrackedDeviceIdFromControllerIndex(int32 ControllerIndex, int32& OutDeviceId)
-{
-	check(IsInGameThread());
 
-	OutDeviceId = -1;
+bool FSteamVRHMD::GetControllerHandPositionAndOrientation( const int32 ControllerIndex, ESteamVRControllerHand Hand, FVector& OutPosition, FQuat& OutOrientation )
+{
+	check( IsInGameThread() );
 
-	if ((ControllerIndex < 0) || (ControllerIndex >= MAX_STEAMVR_CONTROLLERS))
+	if( ( ControllerIndex < 0 ) || ( ControllerIndex >= MAX_STEAMVR_CONTROLLER_PAIRS ) || Hand < ESteamVRControllerHand::Left || Hand > ESteamVRControllerHand::Right )
 	{
 		return false;
 	}
 
-	OutDeviceId = ControllerToDeviceMap[ControllerIndex];
-	
-	return (OutDeviceId != -1);
+	const int32 DeviceId = UnrealControllerIdAndHandToDeviceIdMap[ ControllerIndex ][ (int32)Hand ];
+	return GetTrackedObjectOrientationAndPosition(DeviceId, OutOrientation, OutPosition);
 }
 
+
 TSharedPtr<ISceneViewExtension, ESPMode::ThreadSafe> FSteamVRHMD::GetViewExtension()
 {
 	TSharedPtr<FSteamVRHMD, ESPMode::ThreadSafe> ptr(AsShared());
@@ -949,9 +952,12 @@ void FSteamVRHMD::Startup()
 		}
 
 		// Initialize our controller to device index
-		for (int32 i = 0; i < MAX_STEAMVR_CONTROLLERS; ++i)
+		for (int32 UnrealControllerIndex = 0; UnrealControllerIndex < MAX_STEAMVR_CONTROLLER_PAIRS; ++UnrealControllerIndex )
 		{
-			ControllerToDeviceMap[i] = -1;
+			for( int32 HandIndex = 0; HandIndex < 2; ++HandIndex )
+			{
+				UnrealControllerIdAndHandToDeviceIdMap[ UnrealControllerIndex ][ HandIndex ] = INDEX_NONE;
+			}
 		}
 
 #if PLATFORM_WINDOWS
diff --git a/Engine/Plugins/Runtime/Steam/SteamVR/Source/SteamVR/Private/SteamVRHMD.h b/Engine/Plugins/Runtime/Steam/SteamVR/Source/SteamVR/Private/SteamVRHMD.h
--- a/Engine/Plugins/Runtime/Steam/SteamVR/Source/SteamVR/Private/SteamVRHMD.h
+++ b/Engine/Plugins/Runtime/Steam/SteamVR/Source/SteamVR/Private/SteamVRHMD.h
@@ -4,6 +4,7 @@
 #include "ISteamVRPlugin.h"
 #include "HeadMountedDisplay.h"
 #include "IHeadMountedDisplay.h"
+#include "SteamVRFunctionLibrary.h"
 
 #if PLATFORM_WINDOWS
 #include "AllowWindowsPlatformTypes.h"
@@ -15,21 +16,6 @@
 
 #include "SceneViewExtension.h"
 
-namespace ESteamVRTrackedDeviceType
-{
-	enum Type;
-}
-
-namespace ESteamVRTrackingSpace
-{
-	enum Type;
-}
-
-//@todo steamvr: better association between SteamVRController plugin and SteamVR plugin
-#ifndef MAX_STEAMVR_CONTROLLERS
-	#define MAX_STEAMVR_CONTROLLERS 8
-#endif
-
 /** Stores vectors, in clockwise order, to define soft and hard bounds for Chaperone */
 struct FBoundingQuad
 {
@@ -182,10 +168,11 @@ public:
 	void ShutdownRendering();
 
 	/** Motion Controllers */
-	ESteamVRTrackedDeviceType::Type GetTrackedDeviceType(uint32 DeviceId) const;
-	void GetTrackedDeviceIds(ESteamVRTrackedDeviceType::Type DeviceType, TArray<int32>& TrackedIds);
+	ESteamVRTrackedDeviceType GetTrackedDeviceType(uint32 DeviceId) const;
+	void GetTrackedDeviceIds(ESteamVRTrackedDeviceType DeviceType, TArray<int32>& TrackedIds);
 	bool GetTrackedObjectOrientationAndPosition(uint32 DeviceId, FQuat& CurrentOrientation, FVector& CurrentPosition);
-	bool GetTrackedDeviceIdFromControllerIndex(int32 ControllerIndex, int32& OutDeviceId);
+	bool GetControllerHandPositionAndOrientation( const int32 ControllerIndex, ESteamVRControllerHand Hand, FVector& OutPosition, FQuat& OutOrientation );
+
 
 	/** Chaperone */
 	/** Returns whether or not the player is currently inside the soft bounds */
@@ -201,13 +188,13 @@ public:
 	int32 GetWindowMirrorMode() const { return WindowMirrorMode; }
 
 	/** Sets the tracking space for the returned coordinate system (e.g. standing, sitting) */
-	void SetTrackingSpace(TEnumAsByte<ESteamVRTrackingSpace::Type> NewSpace);
+	void SetTrackingSpace(TEnumAsByte<ESteamVRTrackingSpace> NewSpace);
 
 	/** Returns the tracking space for the returned coordinate system (e.g. standing, sitting) */
-	ESteamVRTrackingSpace::Type GetTrackingSpace() const;
+	ESteamVRTrackingSpace GetTrackingSpace() const;
 
-	/** Sets the map from controller index to tracked device index. */
-	void SetControllerToDeviceMap(int32* InControllerToDeviceMap);
+	/** Sets the map from Unreal controller id and hand index, to tracked device id. */
+	void SetUnrealControllerIdAndHandToDeviceIdMap(int32 InUnrealControllerIdAndHandToDeviceIdMap[ MAX_STEAMVR_CONTROLLER_PAIRS ][ 2 ]);
 
 public:
 	/** Constructor */
@@ -380,8 +367,8 @@ private:
 	/** World units (UU) to Meters scale.  Read from the level, and used to transform positional tracking data */
 	float WorldToMetersScale;
 
-	/** Mapping from Controller index to tracked device index.  Passed in from the controller plugin */
-	int32 ControllerToDeviceMap[MAX_STEAMVR_CONTROLLERS];
+	/** Mapping from Unreal Controller Id and Hand to a tracked device id.  Passed in from the controller plugin */
+	int32 UnrealControllerIdAndHandToDeviceIdMap[MAX_STEAMVR_CONTROLLER_PAIRS][2];
 
 	IRendererModule* RendererModule;
 	ISteamVRPlugin* SteamVRPlugin;
diff --git a/Engine/Plugins/Runtime/Steam/SteamVR/Source/SteamVR/Public/ISteamVRPlugin.h b/Engine/Plugins/Runtime/Steam/SteamVR/Source/SteamVR/Public/ISteamVRPlugin.h
--- a/Engine/Plugins/Runtime/Steam/SteamVR/Source/SteamVR/Public/ISteamVRPlugin.h
+++ b/Engine/Plugins/Runtime/Steam/SteamVR/Source/SteamVR/Public/ISteamVRPlugin.h
@@ -8,6 +8,10 @@
 
 #define STEAMVR_SUPPORTED_PLATFORMS (PLATFORM_WINDOWS)
 
+/** Up to 8 motion controller devices supported (two VR motion controllers per Unreal controller, one for either the left or right hand.) */
+#define MAX_STEAMVR_CONTROLLER_PAIRS 4
+
+
 /**
  * The public interface to this module.  In most cases, this interface is only public to sibling modules 
  * within this plugin.
@@ -55,7 +59,7 @@ public:
 	 * Update the Controller to Device mapping.
 	 * The controller passes this to the HMD.
 	 */
-	virtual void SetControllerToDeviceMap(int32* InControllerToDeviceMap)=0;
+	virtual void SetUnrealControllerIdAndHandToDeviceIdMap(int32 InUnrealControllerIdAndHandToDeviceIdMap[ MAX_STEAMVR_CONTROLLER_PAIRS ][ 2 ] ) = 0;
 
 private:
 };
diff --git a/Engine/Plugins/Runtime/Steam/SteamVR/Source/SteamVRController/Private/SteamVRController.cpp b/Engine/Plugins/Runtime/Steam/SteamVR/Source/SteamVRController/Private/SteamVRController.cpp
--- a/Engine/Plugins/Runtime/Steam/SteamVR/Source/SteamVRController/Private/SteamVRController.cpp
+++ b/Engine/Plugins/Runtime/Steam/SteamVR/Source/SteamVRController/Private/SteamVRController.cpp
@@ -5,14 +5,6 @@
 
 DEFINE_LOG_CATEGORY_STATIC(LogSteamVRController, Log, All);
 
-/** @todo steamvr - do something about this define */
-#ifndef MAX_STEAMVR_CONTROLLERS
-	#define MAX_STEAMVR_CONTROLLERS 8
-#endif
-
-/** Max number of controller buttons.  Must be < 256*/
-#define MAX_NUM_CONTROLLER_BUTTONS 9
-
 /** Controller axis mappings. @todo steamvr: should enumerate rather than hard code */
 #define TOUCHPAD_AXIS	0
 #define TRIGGER_AXIS	1
@@ -23,25 +15,74 @@ DEFINE_LOG_CATEGORY_STATIC(LogSteamVRController, Log, All);
 #define TOUCHPAD_DEADZONE  0.0f
 
 
+
 class FSteamVRController : public IInputDevice
 {
 
 public:
 
+	/** The maximum number of Unreal controllers.  Each Unreal controller represents a pair of motion controller devices */
+	static const int32 MaxUnrealControllers = MAX_STEAMVR_CONTROLLER_PAIRS;
+
+	/** Total number of motion controllers we'll support */
+	static const int32 MaxControllers = MaxUnrealControllers * 2;
+
+	/** Left and right hands for Steam VR motion controller devices */
+	enum class EControllerHand
+	{
+		Left,
+		Right,
+
+		// ...
+		TotalHandCount
+	};
+
+
+	/**
+	 * Buttons on the SteamVR controller
+	 */
+	struct ESteamVRControllerButton
+	{
+		enum Type
+		{
+			System,
+			ApplicationMenu,
+			TouchPadPress,
+			TouchPadTouch,
+			TriggerPress,
+			Grip,
+			TouchPadUp,
+			TouchPadDown,
+			TouchPadLeft,
+			TouchPadRight,
+
+			/** Max number of controller buttons.  Must be < 256 */
+			TotalButtonCount
+		};
+	};
+
+
 	FSteamVRController(const TSharedRef< FGenericApplicationMessageHandler >& InMessageHandler)
 		: MessageHandler(InMessageHandler),
 		  SteamVRPlugin(nullptr)
 	{
 		FMemory::Memzero(ControllerStates, sizeof(ControllerStates));
+		for( int32 ControllerPairIndex = 0; ControllerPairIndex < MaxUnrealControllers; ++ControllerPairIndex )
+		{
+			for( int32 HandIndex = 0; HandIndex < (int32)EControllerHand::TotalHandCount; ++HandIndex )
+			{
+				ControllerStates[ UnrealControllerIdToControllerIndex( ControllerPairIndex, (EControllerHand)HandIndex ) ].Hand = (EControllerHand)HandIndex;
+			}
+		}
 
 		for (int32 i=0; i < vr::k_unMaxTrackedDeviceCount; ++i)
 		{
-			DeviceToControllerMap[i] = -1;
+			DeviceToControllerMap[i] = INDEX_NONE;
 		}
 
-		for (int32 i=0; i < MAX_STEAMVR_CONTROLLERS; ++i)
+		for (int32 i=0; i < MaxControllers; ++i)
 		{
-			ControllerToDeviceMap[i] = -1;
+			ControllerToDeviceMap[i] = INDEX_NONE;
 		}
 
 		NumControllersMapped = 0;
@@ -49,17 +90,27 @@ public:
 		InitialButtonRepeatDelay = 0.2f;
 		ButtonRepeatDelay = 0.1f;
 
-		// set up mapping	//@todo steamvr: these are all presses, do we want to map touches as well?
-
-		Buttons[0] = EControllerButtons::SpecialLeft;				// system
-		Buttons[1] = EControllerButtons::SpecialRight;				// application menu
-		Buttons[2] = EControllerButtons::LeftThumb;					// touchpad press
-		Buttons[3] = EControllerButtons::LeftTriggerThreshold;		// trigger press
-		Buttons[4] = EControllerButtons::BackLeft;					// grip
-		Buttons[5] = EControllerButtons::LeftStickUp;				// touchpad d-pad up
-		Buttons[6] = EControllerButtons::LeftStickDown;				// touchpad d-pad down
-		Buttons[7] = EControllerButtons::LeftStickLeft;				// touchpad d-pad left
-		Buttons[8] = EControllerButtons::LeftStickRight;			// touchpad d-pad right
+		Buttons[ (int32)EControllerHand::Left ][ ESteamVRControllerButton::System ] = EControllerButtons::SpecialLeft;
+		Buttons[ (int32)EControllerHand::Left ][ ESteamVRControllerButton::ApplicationMenu ] = EControllerButtons::LeftShoulder;
+		Buttons[ (int32)EControllerHand::Left ][ ESteamVRControllerButton::TouchPadPress ] = EControllerButtons::LeftThumb;
+		Buttons[ (int32)EControllerHand::Left ][ ESteamVRControllerButton::TouchPadTouch ] = EControllerButtons::Touch0;
+		Buttons[ (int32)EControllerHand::Left ][ ESteamVRControllerButton::TriggerPress ] = EControllerButtons::LeftTriggerThreshold;
+		Buttons[ (int32)EControllerHand::Left ][ ESteamVRControllerButton::Grip ] = EControllerButtons::BackLeft;
+		Buttons[ (int32)EControllerHand::Left ][ ESteamVRControllerButton::TouchPadUp ] = EControllerButtons::LeftStickUp;
+		Buttons[ (int32)EControllerHand::Left ][ ESteamVRControllerButton::TouchPadDown ] = EControllerButtons::LeftStickDown;
+		Buttons[ (int32)EControllerHand::Left ][ ESteamVRControllerButton::TouchPadLeft ] = EControllerButtons::LeftStickLeft;
+		Buttons[ (int32)EControllerHand::Left ][ ESteamVRControllerButton::TouchPadRight ] = EControllerButtons::LeftStickRight;
+
+		Buttons[ (int32)EControllerHand::Right ][ ESteamVRControllerButton::System ] = EControllerButtons::SpecialRight;
+		Buttons[ (int32)EControllerHand::Right ][ ESteamVRControllerButton::ApplicationMenu ] = EControllerButtons::RightShoulder;
+		Buttons[ (int32)EControllerHand::Right ][ ESteamVRControllerButton::TouchPadPress ] = EControllerButtons::RightThumb;
+		Buttons[ (int32)EControllerHand::Right ][ ESteamVRControllerButton::TouchPadTouch ] = EControllerButtons::Touch1;
+		Buttons[ (int32)EControllerHand::Right ][ ESteamVRControllerButton::TriggerPress ] = EControllerButtons::RightTriggerThreshold;
+		Buttons[ (int32)EControllerHand::Right ][ ESteamVRControllerButton::Grip ] = EControllerButtons::BackRight;
+		Buttons[ (int32)EControllerHand::Right ][ ESteamVRControllerButton::TouchPadUp ] = EControllerButtons::RightStickUp;
+		Buttons[ (int32)EControllerHand::Right ][ ESteamVRControllerButton::TouchPadDown ] = EControllerButtons::RightStickDown;
+		Buttons[ (int32)EControllerHand::Right ][ ESteamVRControllerButton::TouchPadLeft ] = EControllerButtons::RightStickLeft;
+		Buttons[ (int32)EControllerHand::Right ][ ESteamVRControllerButton::TouchPadRight ] = EControllerButtons::RightStickRight;
 	}
 
 	virtual ~FSteamVRController()
@@ -72,7 +123,7 @@ public:
 
 	virtual void SendControllerEvents() override
 	{
-		vr::VRControllerState_t ControllerState;
+		vr::VRControllerState_t VRControllerState;
 
 		vr::IVRSystem* VRSystem = GetVRSystem();
 
@@ -89,10 +140,10 @@ public:
 				}
 
 				// update the mappings if this is a new device
-				if (DeviceToControllerMap[DeviceIndex] == -1)
+				if (DeviceToControllerMap[DeviceIndex] == INDEX_NONE )
 				{
 					// don't map too many controllers
-					if (NumControllersMapped >= MAX_STEAMVR_CONTROLLERS)
+					if (NumControllersMapped >= MaxControllers)
 					{
 						continue;
 					}
@@ -102,141 +153,170 @@ public:
 					++NumControllersMapped;
 
 					// update the SteamVR plugin with the new mapping
-					SteamVRPlugin->SetControllerToDeviceMap(ControllerToDeviceMap);
+					{
+						int32 UnrealControllerIdAndHandToDeviceIdMap[ MaxUnrealControllers ][ (int32)EControllerHand::TotalHandCount ];
+						for( int32 UnrealControllerIndex = 0; UnrealControllerIndex < MaxUnrealControllers; ++UnrealControllerIndex )
+						{
+							for( int32 HandIndex = 0; HandIndex < (int32)EControllerHand::TotalHandCount; ++HandIndex )
+							{
+								UnrealControllerIdAndHandToDeviceIdMap[ UnrealControllerIndex ][ HandIndex ] = ControllerToDeviceMap[ UnrealControllerIdToControllerIndex( UnrealControllerIndex, (EControllerHand)HandIndex ) ];
+							}
+						}
+						SteamVRPlugin->SetUnrealControllerIdAndHandToDeviceIdMap( UnrealControllerIdAndHandToDeviceIdMap );
+					}
 				}
 
 				// get the controller index for this device
 				int32 ControllerIndex = DeviceToControllerMap[DeviceIndex];
+				FControllerState& ControllerState = ControllerStates[ ControllerIndex ];
 
-				if (VRSystem->GetControllerState(DeviceIndex, &ControllerState))
+				if (VRSystem->GetControllerState(DeviceIndex, &VRControllerState))
 				{
-					if (ControllerState.unPacketNum != ControllerStates[ControllerIndex].PacketNum )
+					if (VRControllerState.unPacketNum != ControllerState.PacketNum )
 					{
-						bool CurrentStates[MAX_NUM_CONTROLLER_BUTTONS] = {0};
+						bool CurrentStates[ ESteamVRControllerButton::TotalButtonCount ] = {0};
 
 						// Get the current state of all buttons
-						CurrentStates[0] = !!(ControllerState.ulButtonPressed & vr::ButtonMaskFromId(vr::k_EButton_System));
-						CurrentStates[1] = !!(ControllerState.ulButtonPressed & vr::ButtonMaskFromId(vr::k_EButton_ApplicationMenu));
-						CurrentStates[2] = !!(ControllerState.ulButtonPressed & vr::ButtonMaskFromId(vr::k_EButton_SteamVR_Touchpad));
-						CurrentStates[3] = !!(ControllerState.ulButtonPressed & vr::ButtonMaskFromId(vr::k_EButton_SteamVR_Trigger));
-						CurrentStates[4] = !!(ControllerState.ulButtonPressed & vr::ButtonMaskFromId(vr::k_EButton_Grip));
-
-						const bool bTouchPadTouched = (ControllerState.ulButtonTouched & vr::ButtonMaskFromId(vr::k_EButton_SteamVR_Touchpad)) == vr::ButtonMaskFromId(vr::k_EButton_SteamVR_Touchpad);
-						const bool bTouchPadPressed = (ControllerState.ulButtonPressed & vr::ButtonMaskFromId(vr::k_EButton_SteamVR_Touchpad)) == vr::ButtonMaskFromId(vr::k_EButton_SteamVR_Touchpad);
-
-						// If the touchpad isn't currently pressed, zero put both of the axes
-						if (!bTouchPadTouched || 
-							!bTouchPadPressed)	// @todo: Hack to workaround broken hardware where Touched event is not reported correctly
+						CurrentStates[ ESteamVRControllerButton::System ] = !!(VRControllerState.ulButtonPressed & vr::ButtonMaskFromId(vr::k_EButton_System));
+						CurrentStates[ ESteamVRControllerButton::ApplicationMenu ] = !!(VRControllerState.ulButtonPressed & vr::ButtonMaskFromId(vr::k_EButton_ApplicationMenu));
+						CurrentStates[ ESteamVRControllerButton::TouchPadPress ] = !!(VRControllerState.ulButtonPressed & vr::ButtonMaskFromId(vr::k_EButton_SteamVR_Touchpad));
+						CurrentStates[ ESteamVRControllerButton::TouchPadTouch ] = !!( VRControllerState.ulButtonTouched & vr::ButtonMaskFromId( vr::k_EButton_SteamVR_Touchpad ) );
+						CurrentStates[ ESteamVRControllerButton::TriggerPress ] = !!(VRControllerState.ulButtonPressed & vr::ButtonMaskFromId(vr::k_EButton_SteamVR_Trigger));
+						CurrentStates[ ESteamVRControllerButton::Grip ] = !!(VRControllerState.ulButtonPressed & vr::ButtonMaskFromId(vr::k_EButton_Grip));
+
+						// If the touchpad isn't currently pressed or touched, zero put both of the axes
+						if (!CurrentStates[ ESteamVRControllerButton::TouchPadTouch ] ||	// @todo steamvr: This should actually be a '&&', not an '||'.  For now we require button to be pressed to have non-zeroed data.  This is hack to workaround faulty hardware that reports the touchpad as pressed all of the time
+							!CurrentStates[ ESteamVRControllerButton::TouchPadPress ] )
 						{
-							ControllerState.rAxis[TOUCHPAD_AXIS].y = 0.0f;
-							ControllerState.rAxis[TOUCHPAD_AXIS].x = 0.0f;
+							VRControllerState.rAxis[TOUCHPAD_AXIS].y = 0.0f;
+							VRControllerState.rAxis[TOUCHPAD_AXIS].x = 0.0f;
 						}
 
-						CurrentStates[5] = !!(ControllerState.rAxis[TOUCHPAD_AXIS].y > TOUCHPAD_DEADZONE);
-						CurrentStates[6] = !!(ControllerState.rAxis[TOUCHPAD_AXIS].y < -TOUCHPAD_DEADZONE);
-						CurrentStates[7] = !!(ControllerState.rAxis[TOUCHPAD_AXIS].x < -TOUCHPAD_DEADZONE);
-						CurrentStates[8] = !!(ControllerState.rAxis[TOUCHPAD_AXIS].x > TOUCHPAD_DEADZONE);
+						CurrentStates[ ESteamVRControllerButton::TouchPadUp ] = !!(VRControllerState.rAxis[TOUCHPAD_AXIS].y > TOUCHPAD_DEADZONE);
+						CurrentStates[ ESteamVRControllerButton::TouchPadDown ] = !!(VRControllerState.rAxis[TOUCHPAD_AXIS].y < -TOUCHPAD_DEADZONE);
+						CurrentStates[ ESteamVRControllerButton::TouchPadLeft ] = !!(VRControllerState.rAxis[TOUCHPAD_AXIS].x < -TOUCHPAD_DEADZONE);
+						CurrentStates[ ESteamVRControllerButton::TouchPadRight ] = !!(VRControllerState.rAxis[TOUCHPAD_AXIS].x > TOUCHPAD_DEADZONE);
 
-						if (ControllerStates[ControllerIndex].TouchPadXAnalog != ControllerState.rAxis[TOUCHPAD_AXIS].x)
+						if ( ControllerState.TouchPadXAnalog != VRControllerState.rAxis[TOUCHPAD_AXIS].x)
 						{
-							MessageHandler->OnControllerAnalog(EControllerButtons::LeftAnalogX, ControllerIndex, ControllerState.rAxis[TOUCHPAD_AXIS].x);
-							ControllerStates[ControllerIndex].TouchPadXAnalog = ControllerState.rAxis[TOUCHPAD_AXIS].x;
+							MessageHandler->OnControllerAnalog(EControllerButtons::LeftAnalogX, ControllerIndex, VRControllerState.rAxis[TOUCHPAD_AXIS].x);
+							ControllerState.TouchPadXAnalog = VRControllerState.rAxis[TOUCHPAD_AXIS].x;
 						}
 
-						if (ControllerStates[ControllerIndex].TouchPadYAnalog != ControllerState.rAxis[TOUCHPAD_AXIS].y)
+						if ( ControllerState.TouchPadYAnalog != VRControllerState.rAxis[TOUCHPAD_AXIS].y)
 						{
-							MessageHandler->OnControllerAnalog(EControllerButtons::LeftAnalogY, ControllerIndex, ControllerState.rAxis[TOUCHPAD_AXIS].y);
-							ControllerStates[ControllerIndex].TouchPadYAnalog = ControllerState.rAxis[TOUCHPAD_AXIS].y;
+							MessageHandler->OnControllerAnalog(EControllerButtons::LeftAnalogY, ControllerIndex, VRControllerState.rAxis[TOUCHPAD_AXIS].y);
+							ControllerState.TouchPadYAnalog = VRControllerState.rAxis[TOUCHPAD_AXIS].y;
 						}
 
-						if (ControllerStates[ControllerIndex].TriggerAnalog != ControllerState.rAxis[TRIGGER_AXIS].x)
+						if ( ControllerState.TriggerAnalog != VRControllerState.rAxis[TRIGGER_AXIS].x)
 						{
-							MessageHandler->OnControllerAnalog(EControllerButtons::LeftTriggerAnalog, ControllerIndex, ControllerState.rAxis[TRIGGER_AXIS].x);
-							ControllerStates[ControllerIndex].TriggerAnalog = ControllerState.rAxis[TRIGGER_AXIS].x;
+							MessageHandler->OnControllerAnalog(EControllerButtons::LeftTriggerAnalog, ControllerIndex, VRControllerState.rAxis[TRIGGER_AXIS].x);
+							ControllerState.TriggerAnalog = VRControllerState.rAxis[TRIGGER_AXIS].x;
 						}
 
 						// For each button check against the previous state and send the correct message if any
-						for (int32 ButtonIndex = 0; ButtonIndex < MAX_NUM_CONTROLLER_BUTTONS; ++ButtonIndex)
+						for (int32 ButtonIndex = 0; ButtonIndex < ESteamVRControllerButton::TotalButtonCount; ++ButtonIndex)
 						{
-							if (CurrentStates[ButtonIndex] != ControllerStates[ControllerIndex].ButtonStates[ButtonIndex])
+							if (CurrentStates[ButtonIndex] != ControllerState.ButtonStates[ButtonIndex])
 							{
 								if (CurrentStates[ButtonIndex])
 								{
-									MessageHandler->OnControllerButtonPressed(Buttons[ButtonIndex], ControllerIndex, false);
+									MessageHandler->OnControllerButtonPressed( Buttons[ (int32)ControllerState.Hand ][ ButtonIndex ], ControllerIndex, false );
 								}
 								else
 								{
-									MessageHandler->OnControllerButtonReleased(Buttons[ButtonIndex], ControllerIndex, false);
+									MessageHandler->OnControllerButtonReleased( Buttons[ (int32)ControllerState.Hand ][ ButtonIndex ], ControllerIndex, false );
 								}
 
 								if (CurrentStates[ButtonIndex] != 0)
 								{
 									// this button was pressed - set the button's NextRepeatTime to the InitialButtonRepeatDelay
-									ControllerStates[ControllerIndex].NextRepeatTime[ButtonIndex] = CurrentTime + InitialButtonRepeatDelay;
+									ControllerState.NextRepeatTime[ButtonIndex] = CurrentTime + InitialButtonRepeatDelay;
 								}
 							}
 
 							// Update the state for next time
-							ControllerStates[ControllerIndex].ButtonStates[ButtonIndex] = CurrentStates[ButtonIndex];
+							ControllerState.ButtonStates[ButtonIndex] = CurrentStates[ButtonIndex];
 						}
 
-						ControllerStates[ControllerIndex].PacketNum = ControllerState.unPacketNum;
+						ControllerState.PacketNum = VRControllerState.unPacketNum;
 					}
 				}
 
-				for (int32 ButtonIndex = 0; ButtonIndex < MAX_NUM_CONTROLLER_BUTTONS; ++ButtonIndex)
+				for (int32 ButtonIndex = 0; ButtonIndex < ESteamVRControllerButton::TotalButtonCount; ++ButtonIndex)
 				{
-					if (ControllerStates[ControllerIndex].ButtonStates[ButtonIndex] != 0 && ControllerStates[ControllerIndex].NextRepeatTime[ButtonIndex] <= CurrentTime)
+					if ( ControllerState.ButtonStates[ButtonIndex] != 0 && ControllerState.NextRepeatTime[ButtonIndex] <= CurrentTime)
 					{
-						MessageHandler->OnControllerButtonPressed(Buttons[ButtonIndex], ControllerIndex, true);
+						MessageHandler->OnControllerButtonPressed( Buttons[ (int32)ControllerState.Hand ][ ButtonIndex ], ControllerIndex, true );
 
 						// set the button's NextRepeatTime to the ButtonRepeatDelay
-						ControllerStates[ControllerIndex].NextRepeatTime[ButtonIndex] = CurrentTime + ButtonRepeatDelay;
+						ControllerState.NextRepeatTime[ButtonIndex] = CurrentTime + ButtonRepeatDelay;
 					}
 				}
 			}
 		}
 	}
+
+
+	int32 UnrealControllerIdToControllerIndex( const int32 UnrealControllerId, const EControllerHand Hand ) const
+	{
+		return UnrealControllerId * (int32)EControllerHand::TotalHandCount + (int32)Hand;
+	}
+
 	
-	void SetChannelValue(int32 ControllerId, FForceFeedbackChannelType ChannelType, float Value) override
+	void SetChannelValue(int32 UnrealControllerId, FForceFeedbackChannelType ChannelType, float Value) override
 	{
-		// Skip unless this is the left large channel, which we consider to be the only SteamVRController feedback channel
-		if (ChannelType != FForceFeedbackChannelType::LEFT_LARGE)
+		// Skip unless this is the left or right large channel, which we consider to be the only SteamVRController feedback channel
+		if( ChannelType != FForceFeedbackChannelType::LEFT_LARGE && ChannelType != FForceFeedbackChannelType::RIGHT_LARGE )
 		{
 			return;
 		}
 
-		if ((ControllerId >= 0) && (ControllerId < MAX_STEAMVR_CONTROLLERS))
+		const EControllerHand Hand = ( ChannelType == FForceFeedbackChannelType::LEFT_LARGE ) ? EControllerHand::Left : EControllerHand::Right;
+		const int32 ControllerIndex = UnrealControllerIdToControllerIndex( UnrealControllerId, Hand );
+
+		if ((ControllerIndex >= 0) && ( ControllerIndex < MaxControllers))
 		{
-			FControllerState& ControllerState = ControllerStates[ControllerId];
+			FControllerState& ControllerState = ControllerStates[ ControllerIndex ];
 
-			ControllerState.VibeValues.LeftLarge = Value;
+			ControllerState.ForceFeedbackValue = Value;
 
-			UpdateVibration(ControllerId);
+			UpdateVibration( ControllerIndex );
 		}
 	}
 
-	void SetChannelValues(int32 ControllerId, const FForceFeedbackValues &Values) override
+	void SetChannelValues(int32 UnrealControllerId, const FForceFeedbackValues& Values) override
 	{
-		if ((ControllerId >= 0) && (ControllerId < MAX_STEAMVR_CONTROLLERS))
+		const int32 LeftControllerIndex = UnrealControllerIdToControllerIndex( UnrealControllerId, EControllerHand::Left );
+		if ((LeftControllerIndex >= 0) && ( LeftControllerIndex < MaxControllers))
 		{
-			FControllerState& ControllerState = ControllerStates[ControllerId];
-			ControllerState.VibeValues = Values;
+			FControllerState& ControllerState = ControllerStates[ LeftControllerIndex ];
+			ControllerState.ForceFeedbackValue = Values.LeftLarge;
 
-			UpdateVibration(ControllerId);
+			UpdateVibration( LeftControllerIndex );
+		}
+
+		const int32 RightControllerIndex = UnrealControllerIdToControllerIndex( UnrealControllerId, EControllerHand::Right );
+		if( ( RightControllerIndex >= 0 ) && ( RightControllerIndex < MaxControllers ) )
+		{
+			FControllerState& ControllerState = ControllerStates[ RightControllerIndex ];
+			ControllerState.ForceFeedbackValue = Values.RightLarge;
+
+			UpdateVibration( RightControllerIndex );
 		}
 	}
 
-	void UpdateVibration(int32 ControllerId)
+	void UpdateVibration( const int32 ControllerIndex )
 	{
-		// make sure there is a valid device for the controllerid
-		int32 DeviceIndex = ControllerToDeviceMap[ControllerId];
+		// make sure there is a valid device for this controller
+		int32 DeviceIndex = ControllerToDeviceMap[ ControllerIndex ];
 		if (DeviceIndex < 0)
 		{
 			return;
 		}
 
-		const FControllerState& ControllerState = ControllerStates[ControllerId];
+		const FControllerState& ControllerState = ControllerStates[ ControllerIndex ];
 		vr::IVRSystem* VRSystem = GetVRSystem();
 
 		if (VRSystem == nullptr)
@@ -245,7 +325,7 @@ public:
 		}
 
 		// Map the float values from [0,1] to be more reasonable values for the SteamController.  The docs say that [100,2000] are reasonable values
- 		const float LeftIntensity = FMath::Clamp(ControllerState.VibeValues.LeftLarge * 2000.f, 0.f, 2000.f);
+ 		const float LeftIntensity = FMath::Clamp(ControllerState.ForceFeedbackValue * 2000.f, 0.f, 2000.f);
 		if (LeftIntensity > 0.f)
 		{
 			VRSystem->TriggerHapticPulse(DeviceIndex, TOUCHPAD_AXIS, LeftIntensity);
@@ -278,11 +358,15 @@ private:
 		return SteamVRPlugin->GetVRSystem();
 	}
 
+
+
 	struct FControllerState
 	{
+		/** Which hand this controller is representing */
+		EControllerHand Hand;
+
 		/** If packet num matches that on your prior call, then the controller state hasn't been changed since 
-		  * your last call and there is no need to process it.
-		  */
+		  * your last call and there is no need to process it. */
 		uint32 PacketNum;
 
 		/** touchpad analog values */
@@ -293,22 +377,22 @@ private:
 		float TriggerAnalog;
 
 		/** Last frame's button states, so we only send events on edges */
-		bool ButtonStates[MAX_NUM_CONTROLLER_BUTTONS];
+		bool ButtonStates[ ESteamVRControllerButton::TotalButtonCount ];
 
 		/** Next time a repeat event should be generated for each button */
-		double NextRepeatTime[MAX_NUM_CONTROLLER_BUTTONS];
+		double NextRepeatTime[ ESteamVRControllerButton::TotalButtonCount ];
 
-		/** Values for force feedback on this controller.  We only consider the LEFT_LARGE channel for SteamControllers */
-		FForceFeedbackValues VibeValues;
+		/** Value for force feedback on this controller hand */
+		float ForceFeedbackValue;
 	};
 
 	/** Mappings between tracked devices and 0 indexed controllers */
 	int32 NumControllersMapped;
-	int32 ControllerToDeviceMap[MAX_STEAMVR_CONTROLLERS];
+	int32 ControllerToDeviceMap[ MaxControllers ];
 	int32 DeviceToControllerMap[vr::k_unMaxTrackedDeviceCount];
 
 	/** Controller states */
-	FControllerState ControllerStates[MAX_STEAMVR_CONTROLLERS];
+	FControllerState ControllerStates[ MaxControllers ];
 
 	/** Delay before sending a repeat message after a button was first pressed */
 	float InitialButtonRepeatDelay;
@@ -317,7 +401,7 @@ private:
 	float ButtonRepeatDelay;
 
 	/** Mapping of controller buttons */
-	EControllerButtons::Type Buttons[MAX_NUM_CONTROLLER_BUTTONS];
+	EControllerButtons::Type Buttons[ (int32)EControllerHand::TotalHandCount ][ ESteamVRControllerButton::TotalButtonCount ];
 
 	/** handler to send all messages to */
 	TSharedRef<FGenericApplicationMessageHandler> MessageHandler;
diff --git a/Engine/Plugins/Runtime/Steam/SteamVR/Source/SteamVRController/SteamVRController.Build.cs b/Engine/Plugins/Runtime/Steam/SteamVR/Source/SteamVRController/SteamVRController.Build.cs
--- a/Engine/Plugins/Runtime/Steam/SteamVR/Source/SteamVRController/SteamVRController.Build.cs
+++ b/Engine/Plugins/Runtime/Steam/SteamVR/Source/SteamVRController/SteamVRController.Build.cs
@@ -17,11 +17,20 @@ public class SteamVRController : ModuleRules
 			"CoreUObject",
 			"Engine",
 			"InputDevice",
-			"SteamVR",
 			"HeadMountedDisplay"
 		});
 
-        AddThirdPartyPrivateStaticDependencies(Target,
+		DynamicallyLoadedModuleNames.AddRange(new string[]
+		{
+			"SteamVR",
+		});
+
+		PrivateIncludePathModuleNames.AddRange(new string[]
+		{
+			"SteamVR",
+		});
+
+		AddThirdPartyPrivateStaticDependencies(Target,
             "OpenVR"
         );
     }
