Fixed SteamVR devices always reported as disconnected/untracked

diff --git a/Engine/Plugins/Runtime/Steam/SteamVR/Source/SteamVR/Private/SteamVRFunctionLibrary.cpp b/Engine/Plugins/Runtime/Steam/SteamVR/Source/SteamVR/Private/SteamVRFunctionLibrary.cpp
--- a/Engine/Plugins/Runtime/Steam/SteamVR/Source/SteamVR/Private/SteamVRFunctionLibrary.cpp
+++ b/Engine/Plugins/Runtime/Steam/SteamVR/Source/SteamVR/Private/SteamVRFunctionLibrary.cpp
@@ -57,7 +57,7 @@ bool USteamVRFunctionLibrary::GetHandPositionAndOrientation(int32 ControllerInde
 		OutOrientation = DeviceOrientation.Rotator();
 	}
 
-	return false;
+	return RetVal;
 }
 
 void USteamVRFunctionLibrary::SetTrackingSpace(TEnumAsByte<ESteamVRTrackingSpace> NewSpace)
