Shadow variable issues fixed in plugins.

Header changes are needed for VS2015.

diff --git a/Engine/Plugins/Runtime/CableComponent/Source/CableComponent/Classes/CableComponent.h b/Engine/Plugins/Runtime/CableComponent/Source/CableComponent/Classes/CableComponent.h
--- a/Engine/Plugins/Runtime/CableComponent/Source/CableComponent/Classes/CableComponent.h
+++ b/Engine/Plugins/Runtime/CableComponent/Source/CableComponent/Classes/CableComponent.h
@@ -102,9 +102,9 @@ private:
 	/** Solve the cable spring constraints */
 	void SolveConstraints();
 	/** Integrate cable point positions */
-	void VerletIntegrate(float SubstepTime, const FVector& Gravity);
+	void VerletIntegrate(float InSubstepTime, const FVector& Gravity);
 	/** Perform a simulation substep */
-	void PerformSubstep(float SubstepTime, const FVector& Gravity);
+	void PerformSubstep(float InSubstepTime, const FVector& Gravity);
 	/** Get start and end position for the cable */
 	void GetEndPositions(FVector& OutStartPosition, FVector& OutEndPosition);
 	/** Amount of time 'left over' from last tick */
diff --git a/Engine/Plugins/Runtime/OculusRift/Source/OculusRift/Private/OculusRiftRender.cpp b/Engine/Plugins/Runtime/OculusRift/Source/OculusRift/Private/OculusRiftRender.cpp
--- a/Engine/Plugins/Runtime/OculusRift/Source/OculusRift/Private/OculusRiftRender.cpp
+++ b/Engine/Plugins/Runtime/OculusRift/Source/OculusRift/Private/OculusRiftRender.cpp
@@ -571,10 +571,10 @@ void FOculusRiftHMD::UpdateViewport(bool bUseSeparateRenderTarget, const FViewpo
 		{
 			FWidgetPath WidgetPath;
 			TSharedRef<SWidget> Widget = ViewportWidget->AsShared();
-			TSharedPtr<SWindow> Window = FSlateApplication::Get().FindWidgetWindow(Widget, WidgetPath);
-			if (Window.IsValid())
+			TSharedPtr<SWindow> WidgetWindow = FSlateApplication::Get().FindWidgetWindow(Widget, WidgetPath);
+			if (WidgetWindow.IsValid())
 			{
-				Window->SetViewportSizeDrivenByWindow(true);
+				WidgetWindow->SetViewportSizeDrivenByWindow(true);
 			}
 		}
 		return;
diff --git a/Engine/Plugins/Runtime/OculusRift/Source/OculusRift/Private/OculusRiftRenderD3D.cpp b/Engine/Plugins/Runtime/OculusRift/Source/OculusRift/Private/OculusRiftRenderD3D.cpp
--- a/Engine/Plugins/Runtime/OculusRift/Source/OculusRift/Private/OculusRiftRenderD3D.cpp
+++ b/Engine/Plugins/Runtime/OculusRift/Source/OculusRift/Private/OculusRiftRenderD3D.cpp
@@ -153,14 +153,14 @@ FD3D11Texture2DSet* FD3D11Texture2DSet::D3D11CreateTexture2DSet(
 {
 	check(InTextureSet);
 
-	TArray<TRefCountPtr<ID3D11RenderTargetView> > RenderTargetViews;
+	TArray<TRefCountPtr<ID3D11RenderTargetView> > TextureSetRenderTargetViews;
 	FD3D11Texture2DSet* NewTextureSet = new FD3D11Texture2DSet(
 		InD3D11RHI,
 		nullptr,
 		nullptr,
 		false,
 		1,
-		RenderTargetViews,
+		TextureSetRenderTargetViews,
 		/*DepthStencilViews=*/ NULL,
 		InDsDesc.Width,
 		InDsDesc.Height,
diff --git a/Engine/Plugins/Runtime/Steam/SteamVR/Source/SteamVR/Private/SteamVRHMD.cpp b/Engine/Plugins/Runtime/Steam/SteamVR/Source/SteamVR/Private/SteamVRHMD.cpp
--- a/Engine/Plugins/Runtime/Steam/SteamVR/Source/SteamVR/Private/SteamVRHMD.cpp
+++ b/Engine/Plugins/Runtime/Steam/SteamVR/Source/SteamVR/Private/SteamVRHMD.cpp
@@ -917,14 +917,14 @@ void FSteamVRHMD::Startup()
 		float WidthPercentage = ((float)RecommendedWidth / (float)ScreenWidth) * 100.0f;
 		float HeightPercentage = ((float)RecommendedHeight / (float)ScreenHeight) * 100.0f;
 
-		float IdealScreenPercentage = FMath::Max(WidthPercentage, HeightPercentage);
+		float ScreenPercentage = FMath::Max(WidthPercentage, HeightPercentage);
 
 		//@todo steamvr: move out of here
 		static IConsoleVariable* CScrPercVar = IConsoleManager::Get().FindConsoleVariable(TEXT("r.ScreenPercentage"));
 
-		if (FMath::RoundToInt(CScrPercVar->GetFloat()) != FMath::RoundToInt(IdealScreenPercentage))
+		if (FMath::RoundToInt(CScrPercVar->GetFloat()) != FMath::RoundToInt(ScreenPercentage))
 		{
-			CScrPercVar->Set(IdealScreenPercentage);
+			CScrPercVar->Set(ScreenPercentage);
 		}
 
 		// disable vsync
diff --git a/Engine/Plugins/Runtime/Steam/SteamVR/Source/SteamVR/Private/SteamVRHMD.h b/Engine/Plugins/Runtime/Steam/SteamVR/Source/SteamVR/Private/SteamVRHMD.h
--- a/Engine/Plugins/Runtime/Steam/SteamVR/Source/SteamVR/Private/SteamVRHMD.h
+++ b/Engine/Plugins/Runtime/Steam/SteamVR/Source/SteamVR/Private/SteamVRHMD.h
@@ -129,7 +129,7 @@ public:
 		bool IsInitialized() const { return bInitialized; }
 
 		virtual void BeginRendering() = 0;
-		virtual void UpdateViewport(const FViewport& Viewport, FRHIViewport* ViewportRHI) = 0;
+		virtual void UpdateViewport(const FViewport& Viewport, FRHIViewport* InViewportRHI) = 0;
 		virtual void SetNeedReinitRendererAPI() { bNeedReinitRendererAPI = true; }
 
 		virtual void Reset() = 0;
@@ -152,7 +152,7 @@ public:
 
 		virtual void BeginRendering() override;
 		void FinishRendering();
-		virtual void UpdateViewport(const FViewport& Viewport, FRHIViewport* ViewportRHI) override;
+		virtual void UpdateViewport(const FViewport& Viewport, FRHIViewport* InViewportRHI) override;
 		virtual void Reset() override;
 		virtual void Shutdown() override
 		{
diff --git a/Engine/Plugins/Runtime/Steam/SteamVR/Source/SteamVR/Private/SteamVRRender.cpp b/Engine/Plugins/Runtime/Steam/SteamVR/Source/SteamVR/Private/SteamVRRender.cpp
--- a/Engine/Plugins/Runtime/Steam/SteamVR/Source/SteamVR/Private/SteamVRRender.cpp
+++ b/Engine/Plugins/Runtime/Steam/SteamVR/Source/SteamVR/Private/SteamVRRender.cpp
@@ -118,10 +118,10 @@ void FSteamVRHMD::D3D11Bridge::Reset()
 {
 }
 
-void FSteamVRHMD::D3D11Bridge::UpdateViewport(const FViewport& Viewport, FRHIViewport* ViewportRHI)
+void FSteamVRHMD::D3D11Bridge::UpdateViewport(const FViewport& Viewport, FRHIViewport* InViewportRHI)
 {
 	check(IsInGameThread());
-	check(ViewportRHI);
+	check(InViewportRHI);
 
 	const FTexture2DRHIRef& RT = Viewport.GetRenderTargetTexture();
 	check(IsValidRef(RT));
@@ -134,7 +134,7 @@ void FSteamVRHMD::D3D11Bridge::UpdateViewport(const FViewport& Viewport, FRHIVie
 	RenderTargetTexture = (ID3D11Texture2D*)RT->GetNativeResource();
 	RenderTargetTexture->AddRef();
 
-	ViewportRHI->SetCustomPresent(this);
+	InViewportRHI->SetCustomPresent(this);
 }
 
 
