#motioncontroller Adding generic MotionController system, along with a MotionControllerComponent, which allows you to specify a player index and a hand to track, and it will automatically update the component with the appropriate offset

diff --git a/Engine/Plugins/Runtime/Steam/SteamVR/Source/SteamVR/Classes/SteamVRFunctionLibrary.h b/Engine/Plugins/Runtime/Steam/SteamVR/Source/SteamVR/Classes/SteamVRFunctionLibrary.h
--- a/Engine/Plugins/Runtime/Steam/SteamVR/Source/SteamVR/Classes/SteamVRFunctionLibrary.h
+++ b/Engine/Plugins/Runtime/Steam/SteamVR/Source/SteamVR/Classes/SteamVRFunctionLibrary.h
@@ -1,6 +1,7 @@
 // Copyright 1998-2015 Epic Games, Inc. All Rights Reserved.
 
 #pragma once
+#include "IMotionController.h"
 #include "SteamVRFunctionLibrary.generated.h"
 
 /** Defines the class of tracked devices in SteamVR*/
@@ -31,17 +32,6 @@ enum class ESteamVRTrackingSpace
 	Seated
 };
 
-/** The motion controller hand for a specific Unreal controller.  Each Unreal controller can support a pair of devices, one for the left hand and one for the right hand. */
-UENUM(BlueprintType)
-enum class ESteamVRControllerHand
-{
-	/** Left hand */
-	Left,
-
-	/** Right hand */
-	Right,
-};
-
 /**
  * SteamVR Extensions Function Library
  */
@@ -82,7 +72,7 @@ public:
 	 * @return	True if the specified controller index has a valid tracked device ID
 	 */
 	UFUNCTION(BlueprintPure, Category = "SteamVR")
-	static bool GetHandPositionAndOrientation(int32 ControllerIndex, ESteamVRControllerHand Hand, FVector& OutPosition, FRotator& OutOrientation);
+	static bool GetHandPositionAndOrientation(int32 ControllerIndex, EControllerHand Hand, FVector& OutPosition, FRotator& OutOrientation);
 
 	/**
 	 * Sets the tracking space (e.g. sitting or standing), which changes which space tracked positions are returned to
diff --git a/Engine/Plugins/Runtime/Steam/SteamVR/Source/SteamVR/Private/SteamVRFunctionLibrary.cpp b/Engine/Plugins/Runtime/Steam/SteamVR/Source/SteamVR/Private/SteamVRFunctionLibrary.cpp
--- a/Engine/Plugins/Runtime/Steam/SteamVR/Source/SteamVR/Private/SteamVRFunctionLibrary.cpp
+++ b/Engine/Plugins/Runtime/Steam/SteamVR/Source/SteamVR/Private/SteamVRFunctionLibrary.cpp
@@ -1,7 +1,6 @@
 // Copyright 1998-2015 Epic Games, Inc. All Rights Reserved.
 //
 #include "SteamVRPrivatePCH.h"
-#include "SteamVRHMD.h"
 #include "Classes/SteamVRFunctionLibrary.h"
 
 USteamVRFunctionLibrary::USteamVRFunctionLibrary(const FObjectInitializer& ObjectInitializer)
@@ -45,7 +44,7 @@ bool USteamVRFunctionLibrary::GetTrackedDevicePositionAndOrientation(int32 Devic
 	return RetVal;
 }
 
-bool USteamVRFunctionLibrary::GetHandPositionAndOrientation(int32 ControllerIndex, ESteamVRControllerHand Hand, FVector& OutPosition, FRotator& OutOrientation)
+bool USteamVRFunctionLibrary::GetHandPositionAndOrientation(int32 ControllerIndex, EControllerHand Hand, FVector& OutPosition, FRotator& OutOrientation)
 {
 	bool RetVal = false;
 
diff --git a/Engine/Plugins/Runtime/Steam/SteamVR/Source/SteamVR/Private/SteamVRHMD.cpp b/Engine/Plugins/Runtime/Steam/SteamVR/Source/SteamVR/Private/SteamVRHMD.cpp
--- a/Engine/Plugins/Runtime/Steam/SteamVR/Source/SteamVR/Private/SteamVRHMD.cpp
+++ b/Engine/Plugins/Runtime/Steam/SteamVR/Source/SteamVR/Private/SteamVRHMD.cpp
@@ -399,11 +399,11 @@ bool FSteamVRHMD::GetTrackedObjectOrientationAndPosition(uint32 DeviceId, FQuat&
 }
 
 
-bool FSteamVRHMD::GetControllerHandPositionAndOrientation( const int32 ControllerIndex, ESteamVRControllerHand Hand, FVector& OutPosition, FQuat& OutOrientation )
+bool FSteamVRHMD::GetControllerHandPositionAndOrientation( const int32 ControllerIndex, EControllerHand Hand, FVector& OutPosition, FQuat& OutOrientation )
 {
 	check( IsInGameThread() );
 
-	if( ( ControllerIndex < 0 ) || ( ControllerIndex >= MAX_STEAMVR_CONTROLLER_PAIRS ) || Hand < ESteamVRControllerHand::Left || Hand > ESteamVRControllerHand::Right )
+	if ((ControllerIndex < 0) || (ControllerIndex >= MAX_STEAMVR_CONTROLLER_PAIRS) || Hand < EControllerHand::Left || Hand > EControllerHand::Right)
 	{
 		return false;
 	}
diff --git a/Engine/Plugins/Runtime/Steam/SteamVR/Source/SteamVR/Private/SteamVRHMD.h b/Engine/Plugins/Runtime/Steam/SteamVR/Source/SteamVR/Private/SteamVRHMD.h
--- a/Engine/Plugins/Runtime/Steam/SteamVR/Source/SteamVR/Private/SteamVRHMD.h
+++ b/Engine/Plugins/Runtime/Steam/SteamVR/Source/SteamVR/Private/SteamVRHMD.h
@@ -171,7 +171,7 @@ public:
 	ESteamVRTrackedDeviceType GetTrackedDeviceType(uint32 DeviceId) const;
 	void GetTrackedDeviceIds(ESteamVRTrackedDeviceType DeviceType, TArray<int32>& TrackedIds);
 	bool GetTrackedObjectOrientationAndPosition(uint32 DeviceId, FQuat& CurrentOrientation, FVector& CurrentPosition);
-	bool GetControllerHandPositionAndOrientation( const int32 ControllerIndex, ESteamVRControllerHand Hand, FVector& OutPosition, FQuat& OutOrientation );
+	STEAMVR_API bool GetControllerHandPositionAndOrientation( const int32 ControllerIndex, EControllerHand Hand, FVector& OutPosition, FQuat& OutOrientation );
 
 
 	/** Chaperone */
diff --git a/Engine/Plugins/Runtime/Steam/SteamVR/Source/SteamVRController/Private/SteamVRController.cpp b/Engine/Plugins/Runtime/Steam/SteamVR/Source/SteamVRController/Private/SteamVRController.cpp
--- a/Engine/Plugins/Runtime/Steam/SteamVR/Source/SteamVRController/Private/SteamVRController.cpp
+++ b/Engine/Plugins/Runtime/Steam/SteamVR/Source/SteamVRController/Private/SteamVRController.cpp
@@ -1,10 +1,14 @@
 // Copyright 1998-2015 Epic Games, Inc. All Rights Reserved.
 
 #include "SteamVRControllerPrivatePCH.h"
-
+#include "IMotionController.h"
+#include "../../SteamVR/Private/SteamVRHMD.h"
 
 DEFINE_LOG_CATEGORY_STATIC(LogSteamVRController, Log, All);
 
+/** Total number of controllers in a set */
+#define CONTROLLERS_PER_PLAYER	2
+
 /** Controller axis mappings. @todo steamvr: should enumerate rather than hard code */
 #define TOUCHPAD_AXIS	0
 #define TRIGGER_AXIS	1
@@ -21,7 +25,17 @@ namespace SteamVRControllerKeyNames
 	const FGamepadKeyNames::Type Touch1("Steam_Touch_1");
 }
 
-class FSteamVRController : public IInputDevice
+FSteamVRHMD* GetSteamVRHMD()
+{
+	if (GEngine->HMDDevice.IsValid() && (GEngine->HMDDevice->GetHMDDeviceType() == EHMDDeviceType::DT_SteamVR))
+	{
+		return static_cast<FSteamVRHMD*>(GEngine->HMDDevice.Get());
+	}
+
+	return nullptr;
+}
+
+class FSteamVRController : public IInputDevice, public IMotionController
 {
 
 public:
@@ -32,17 +46,6 @@ public:
 	/** Total number of motion controllers we'll support */
 	static const int32 MaxControllers = MaxUnrealControllers * 2;
 
-	/** Left and right hands for Steam VR motion controller devices */
-	enum class EControllerHand
-	{
-		Left,
-		Right,
-
-		// ...
-		TotalHandCount
-	};
-
-
 	/**
 	 * Buttons on the SteamVR controller
 	 */
@@ -74,7 +77,7 @@ public:
 		FMemory::Memzero(ControllerStates, sizeof(ControllerStates));
 		for( int32 ControllerPairIndex = 0; ControllerPairIndex < MaxUnrealControllers; ++ControllerPairIndex )
 		{
-			for( int32 HandIndex = 0; HandIndex < (int32)EControllerHand::TotalHandCount; ++HandIndex )
+			for (int32 HandIndex = 0; HandIndex < CONTROLLERS_PER_PLAYER; ++HandIndex)
 			{
 				ControllerStates[ UnrealControllerIdToControllerIndex( ControllerPairIndex, (EControllerHand)HandIndex ) ].Hand = (EControllerHand)HandIndex;
 			}
@@ -116,10 +119,16 @@ public:
 		Buttons[ (int32)EControllerHand::Right ][ ESteamVRControllerButton::TouchPadDown ] = FGamepadKeyNames::MotionController_Right_FaceButton3;
 		Buttons[ (int32)EControllerHand::Right ][ ESteamVRControllerButton::TouchPadLeft ] = FGamepadKeyNames::MotionController_Right_FaceButton4;
 		Buttons[ (int32)EControllerHand::Right ][ ESteamVRControllerButton::TouchPadRight ] = FGamepadKeyNames::MotionController_Right_FaceButton2;
+
+		IModularFeatures::Get().RegisterModularFeature(GetModularFeatureName(), this);
+
+		//@todo:  fix this.  construction of the controller happens after InitializeMotionControllers(), so we manually insert into the array here.
+		GEngine->MotionControllerDevices.AddUnique(this);
 	}
 
 	virtual ~FSteamVRController()
 	{
+		GEngine->MotionControllerDevices.Remove(this);
 	}
 
 	virtual void Tick( float DeltaTime ) override
@@ -153,16 +162,16 @@ public:
 						continue;
 					}
 
-					DeviceToControllerMap[DeviceIndex] = FMath::FloorToInt(NumControllersMapped / (int32)EControllerHand::TotalHandCount);
+					DeviceToControllerMap[DeviceIndex] = FMath::FloorToInt(NumControllersMapped / CONTROLLERS_PER_PLAYER);
 					ControllerToDeviceMap[NumControllersMapped] = DeviceIndex;
 					++NumControllersMapped;
 
 					// update the SteamVR plugin with the new mapping
 					{
-						int32 UnrealControllerIdAndHandToDeviceIdMap[ MaxUnrealControllers ][ (int32)EControllerHand::TotalHandCount ];
+						int32 UnrealControllerIdAndHandToDeviceIdMap[MaxUnrealControllers][CONTROLLERS_PER_PLAYER];
 						for( int32 UnrealControllerIndex = 0; UnrealControllerIndex < MaxUnrealControllers; ++UnrealControllerIndex )
 						{
-							for( int32 HandIndex = 0; HandIndex < (int32)EControllerHand::TotalHandCount; ++HandIndex )
+							for (int32 HandIndex = 0; HandIndex < CONTROLLERS_PER_PLAYER; ++HandIndex)
 							{
 								UnrealControllerIdAndHandToDeviceIdMap[ UnrealControllerIndex ][ HandIndex ] = ControllerToDeviceMap[ UnrealControllerIdToControllerIndex( UnrealControllerIndex, (EControllerHand)HandIndex ) ];
 							}
@@ -278,7 +287,7 @@ public:
 
 	int32 UnrealControllerIdToControllerIndex( const int32 UnrealControllerId, const EControllerHand Hand ) const
 	{
-		return UnrealControllerId * (int32)EControllerHand::TotalHandCount + (int32)Hand;
+		return UnrealControllerId * CONTROLLERS_PER_PLAYER + (int32)Hand;
 	}
 
 	
@@ -354,6 +363,21 @@ public:
 		MessageHandler = InMessageHandler;
 	}
 
+	virtual bool GetControllerOrientationAndPosition(const int32 ControllerIndex, const EControllerHand DeviceHand, FRotator& OutOrientation, FVector& OutPosition) const
+	{
+		bool RetVal = false;
+
+ 		FSteamVRHMD* SteamVRHMD = GetSteamVRHMD();
+ 		if (SteamVRHMD)
+ 		{
+ 			FQuat DeviceOrientation = FQuat::Identity;
+ 			RetVal = SteamVRHMD->GetControllerHandPositionAndOrientation(ControllerIndex, DeviceHand, OutPosition, DeviceOrientation);
+ 			OutOrientation = DeviceOrientation.Rotator();
+ 		}
+
+		return RetVal;
+	}
+
 	virtual bool Exec( UWorld* InWorld, const TCHAR* Cmd, FOutputDevice& Ar ) override
 	{
 		return false;
@@ -418,7 +442,7 @@ private:
 	float ButtonRepeatDelay;
 
 	/** Mapping of controller buttons */
-	FGamepadKeyNames::Type Buttons[ (int32)EControllerHand::TotalHandCount ][ ESteamVRControllerButton::TotalButtonCount ];
+	FGamepadKeyNames::Type Buttons[ CONTROLLERS_PER_PLAYER ][ ESteamVRControllerButton::TotalButtonCount ];
 
 	/** handler to send all messages to */
 	TSharedRef<FGenericApplicationMessageHandler> MessageHandler;
diff --git a/Engine/Plugins/Runtime/Steam/SteamVR/Source/SteamVRController/SteamVRController.Build.cs b/Engine/Plugins/Runtime/Steam/SteamVR/Source/SteamVRController/SteamVRController.Build.cs
--- a/Engine/Plugins/Runtime/Steam/SteamVR/Source/SteamVRController/SteamVRController.Build.cs
+++ b/Engine/Plugins/Runtime/Steam/SteamVR/Source/SteamVRController/SteamVRController.Build.cs
@@ -9,6 +9,7 @@ public class SteamVRController : ModuleRules
         PrivateIncludePathModuleNames.AddRange(new string[]
 		{
 			"TargetPlatform",
+            "SteamVR"
 		});
 
         PrivateDependencyModuleNames.AddRange(new string[]
@@ -17,18 +18,14 @@ public class SteamVRController : ModuleRules
 			"CoreUObject",
 			"Engine",
 			"InputDevice",
-			"HeadMountedDisplay"
+			"HeadMountedDisplay",
+            "SteamVR"
 		});
 
-		DynamicallyLoadedModuleNames.AddRange(new string[]
-		{
-			"SteamVR",
-		});
-
-		PrivateIncludePathModuleNames.AddRange(new string[]
-		{
-			"SteamVR",
-		});
+// 		DynamicallyLoadedModuleNames.AddRange(new string[]
+// 		{
+// 			"SteamVR",
+// 		});
 
 		AddThirdPartyPrivateStaticDependencies(Target,
             "OpenVR"
diff --git a/Engine/Source/Runtime/Engine/Classes/Engine/Engine.h b/Engine/Source/Runtime/Engine/Classes/Engine/Engine.h
--- a/Engine/Source/Runtime/Engine/Classes/Engine/Engine.h
+++ b/Engine/Source/Runtime/Engine/Classes/Engine/Engine.h
@@ -1496,7 +1496,10 @@ public:
 
 	/** Reference to the HMD device that is attached, if any */
 	TSharedPtr< class IHeadMountedDisplay, ESPMode::ThreadSafe > HMDDevice;
-	
+
+	/** Reference to the Motion Control devices that are attached, if any */
+	TArray < class IMotionController*> MotionControllerDevices;
+
 	/** Triggered when a world is added. */	
 	DECLARE_EVENT_OneParam( UEngine, FWorldAddedEvent , UWorld* );
 	
@@ -2193,6 +2196,13 @@ protected:
 	 */
 	virtual bool InitializeHMDDevice();
 
+	/**
+	 *	Detects and initializes any motion controller devices
+	 *
+	 *	@return true if there are any initialized motion controllers, false otherwise
+	 */
+	virtual bool InitializeMotionControllers();
+
 	/**
 	 *	Record EngineAnalytics information for attached HMD devices
 	 */
diff --git a/Engine/Source/Runtime/Engine/Private/UnrealEngine.cpp b/Engine/Source/Runtime/Engine/Private/UnrealEngine.cpp
--- a/Engine/Source/Runtime/Engine/Private/UnrealEngine.cpp
+++ b/Engine/Source/Runtime/Engine/Private/UnrealEngine.cpp
@@ -37,6 +37,7 @@
 #include "StereoRendering.h"
 #include "IHeadMountedDisplayModule.h"
 #include "IHeadMountedDisplay.h"
+#include "IMotionController.h"
 #include "Scalability.h"
 #include "StatsData.h"
 #include "ScreenRendering.h"
@@ -733,8 +734,9 @@ void UEngine::Init(IEngineLoop* InEngineLoop)
 	// Add to root.
 	AddToRoot();
 
-	// Initialize the HMD, if any
+	// Initialize the HMDs and motion controllers, if any
 	InitializeHMDDevice();
+	InitializeMotionControllers();
 
 	// Disable the screensaver when running the game.
 	if( GIsClient && !GIsEditor )
@@ -2175,6 +2177,17 @@ bool UEngine::InitializeHMDDevice()
 	return StereoRenderingDevice.IsValid();
 }
 
+bool UEngine::InitializeMotionControllers()
+{
+	TArray<IMotionController*> MotionControllers = IModularFeatures::Get().GetModularFeatureImplementations<IMotionController>(IMotionController::GetModularFeatureName());
+	for (auto MotionController : MotionControllers)
+	{
+		MotionControllerDevices.AddUnique(MotionController);
+	}
+
+	return (MotionControllerDevices.Num() > 0);
+}
+
 void UEngine::RecordHMDAnalytics()
 {
 	if( !GIsEditor )
diff --git a/Engine/Source/Runtime/HeadMountedDisplay/Private/MotionControllerComponent.cpp b/Engine/Source/Runtime/HeadMountedDisplay/Private/MotionControllerComponent.cpp
new file mode 100644
--- /dev/null
+++ b/Engine/Source/Runtime/HeadMountedDisplay/Private/MotionControllerComponent.cpp
@@ -0,0 +1,37 @@
+// Copyright 1998-2015 Epic Games, Inc. All Rights Reserved.
+//
+#include "HeadMountedDisplayPrivate.h"
+#include "MotionControllerComponent.h"
+
+UMotionControllerComponent::UMotionControllerComponent(const FObjectInitializer& ObjectInitializer)
+: Super(ObjectInitializer)
+{
+	PrimaryComponentTick.bCanEverTick = true;
+	PrimaryComponentTick.bStartWithTickEnabled = true;
+	PrimaryComponentTick.TickGroup = TG_PrePhysics;
+
+	PlayerIndex = 0;
+	Hand = EControllerHand::Left;
+}
+
+void UMotionControllerComponent::TickComponent(float DeltaTime, enum ELevelTick TickType, FActorComponentTickFunction *ThisTickFunction)
+{
+	Super::TickComponent(DeltaTime, TickType, ThisTickFunction);
+
+	bTracked = false;
+	FVector Position = FVector::ZeroVector;
+	FRotator Orientation = FRotator::ZeroRotator;
+
+	if (PlayerIndex != INDEX_NONE)
+	{
+		for (auto MotionController : GEngine->MotionControllerDevices)
+		{
+			if ((MotionController != nullptr) && MotionController->GetControllerOrientationAndPosition(PlayerIndex, Hand, Orientation, Position))
+			{
+				SetRelativeLocationAndRotation(Position, Orientation);
+				bTracked = true;
+				break;
+			}
+		}
+	}
+}
diff --git a/Engine/Source/Runtime/HeadMountedDisplay/Public/IMotionController.h b/Engine/Source/Runtime/HeadMountedDisplay/Public/IMotionController.h
new file mode 100644
--- /dev/null
+++ b/Engine/Source/Runtime/HeadMountedDisplay/Public/IMotionController.h
@@ -0,0 +1,40 @@
+// Copyright 1998-2015 Epic Games, Inc. All Rights Reserved.
+
+#pragma once
+
+/** Defines the controller hands for tracking.  Could be expanded, as needed, to facilitate non-handheld controllers */
+UENUM(BlueprintType)
+enum class EControllerHand
+{
+	Left,
+	Right
+};
+
+/**
+ * Motion Controller device interface
+ *
+ * NOTE:  This intentionally does NOT derive from IInputDeviceModule, to allow a clean separation for devices which exclusively track motion with no tactile input
+ * NOTE:  You must MANUALLY call IModularFeatures::Get().RegisterModularFeature( GetModularFeatureName(), this ) in your implementation!  This allows motion controllers
+ *			to be both piggy-backed off HMD devices which support them, as well as standing alone.
+ */
+
+class HEADMOUNTEDDISPLAY_API IMotionController : public IModularFeature
+{
+public:
+	static FName GetModularFeatureName()
+	{
+		static FName FeatureName = FName(TEXT("MotionController"));
+		return FeatureName;
+	}
+
+	/**
+	 * Returns the calibration-space orientation of the requested controller's hand.
+	 *
+	 * @param ControllerIndex	The Unreal controller (player) index of the contoller set
+	 * @param DeviceHand		Which hand, within the controller set for the player, to get the orientation and position for
+	 * @param OutOrientation	(out) If tracked, the orientation (in calibrated-space) of the controller in the specified hand
+	 * @param OutPosition		(out) If tracked, the position (in calibrated-space) of the controller in the specified hand
+	 * @return					True if the device requested is valid and tracked, false otherwise
+	 */
+	virtual bool GetControllerOrientationAndPosition(const int32 ControllerIndex, const EControllerHand DeviceHand, FRotator& OutOrientation, FVector& OutPosition) const = 0;
+};
diff --git a/Engine/Source/Runtime/HeadMountedDisplay/Public/MotionControllerComponent.h b/Engine/Source/Runtime/HeadMountedDisplay/Public/MotionControllerComponent.h
new file mode 100644
--- /dev/null
+++ b/Engine/Source/Runtime/HeadMountedDisplay/Public/MotionControllerComponent.h
@@ -0,0 +1,32 @@
+// Copyright 1998-2015 Epic Games, Inc. All Rights Reserved.
+
+#pragma once
+#include "IMotionController.h"
+#include "MotionControllerComponent.generated.h"
+
+UCLASS(MinimalAPI, meta = (BlueprintSpawnableComponent), ClassGroup = MotionController)
+class UMotionControllerComponent : public UPrimitiveComponent
+{
+	GENERATED_UCLASS_BODY()
+
+	/** Which player index this motion controller should automatically follow */
+	UPROPERTY(EditAnywhere, BlueprintReadWrite, Category = "MotionController")
+	int32 PlayerIndex;
+
+	/** Which hand this component should automatically follow */
+	UPROPERTY(EditAnywhere, BlueprintReadWrite, Category = "MotionController")
+	TEnumAsByte<EControllerHand> Hand;
+
+	void TickComponent(float DeltaTime, enum ELevelTick TickType, FActorComponentTickFunction *ThisTickFunction) override;
+
+	/** Whether or not this component had a valid tracked device this frame */
+	UFUNCTION(BlueprintPure, Category = "MotionController")
+	bool IsTracked() const
+	{
+		return bTracked;
+	}
+
+private:
+	/** Whether or not this component had a valid tracked controller associated with it this frame*/
+	bool bTracked;
+};
