Moved GetSteamVRHMD in SteamVRController.cpp into the FSteamVRController (was causing a multiply defined symbol error on some builds)

diff --git a/Engine/Plugins/Runtime/Steam/SteamVR/Source/SteamVRController/Private/SteamVRController.cpp b/Engine/Plugins/Runtime/Steam/SteamVR/Source/SteamVRController/Private/SteamVRController.cpp
--- a/Engine/Plugins/Runtime/Steam/SteamVR/Source/SteamVRController/Private/SteamVRController.cpp
+++ b/Engine/Plugins/Runtime/Steam/SteamVR/Source/SteamVRController/Private/SteamVRController.cpp
@@ -25,18 +25,19 @@ namespace SteamVRControllerKeyNames
 	const FGamepadKeyNames::Type Touch1("Steam_Touch_1");
 }
 
-FSteamVRHMD* GetSteamVRHMD()
-{
-	if (GEngine->HMDDevice.IsValid() && (GEngine->HMDDevice->GetHMDDeviceType() == EHMDDeviceType::DT_SteamVR))
-	{
-		return static_cast<FSteamVRHMD*>(GEngine->HMDDevice.Get());
-	}
 
-	return nullptr;
-}
 
 class FSteamVRController : public IInputDevice, public IMotionController
 {
+	FSteamVRHMD* GetSteamVRHMD() const
+	{
+		if (GEngine->HMDDevice.IsValid() && (GEngine->HMDDevice->GetHMDDeviceType() == EHMDDeviceType::DT_SteamVR))
+		{
+			return static_cast<FSteamVRHMD*>(GEngine->HMDDevice.Get());
+		}
+
+		return nullptr;
+	}
 
 public:
 
