Fixed non-unity build error

diff --git a/Engine/Plugins/Runtime/Steam/SteamVR/Source/SteamVR/Private/SteamVRFunctionLibrary.cpp b/Engine/Plugins/Runtime/Steam/SteamVR/Source/SteamVR/Private/SteamVRFunctionLibrary.cpp
--- a/Engine/Plugins/Runtime/Steam/SteamVR/Source/SteamVR/Private/SteamVRFunctionLibrary.cpp
+++ b/Engine/Plugins/Runtime/Steam/SteamVR/Source/SteamVR/Private/SteamVRFunctionLibrary.cpp
@@ -2,6 +2,7 @@
 //
 #include "SteamVRPrivatePCH.h"
 #include "Classes/SteamVRFunctionLibrary.h"
+#include "SteamVRHMD.h"
 
 USteamVRFunctionLibrary::USteamVRFunctionLibrary(const FObjectInitializer& ObjectInitializer)
 : Super(ObjectInitializer)
