#steamvr Fix for hands not matching the controller inputs

diff --git a/Engine/Plugins/Runtime/Steam/SteamVR/Source/SteamVRController/Private/SteamVRController.cpp b/Engine/Plugins/Runtime/Steam/SteamVR/Source/SteamVRController/Private/SteamVRController.cpp
--- a/Engine/Plugins/Runtime/Steam/SteamVR/Source/SteamVRController/Private/SteamVRController.cpp
+++ b/Engine/Plugins/Runtime/Steam/SteamVR/Source/SteamVRController/Private/SteamVRController.cpp
@@ -19,6 +19,15 @@ DEFINE_LOG_CATEGORY_STATIC(LogSteamVRController, Log, All);
 //
 #define TOUCHPAD_DEADZONE  0.0f
 
+// Controls whether or not we need to swap the input routing for the hands, for debugging
+static TAutoConsoleVariable<int32> CVarSwapHands(
+	TEXT("vr.SwapMotionControllerInput"),
+	0,
+	TEXT("This command allows you to swap the button / axis input handedness for the input controller, for debugging purposes.\n")
+	TEXT(" 0: don't swap (default)\n")
+	TEXT(" 1: swap left and right buttons"),
+	ECVF_Cheat);
+
 namespace SteamVRControllerKeyNames
 {
 	const FGamepadKeyNames::Type Touch0("Steam_Touch_0");
@@ -76,13 +85,6 @@ public:
 		  SteamVRPlugin(nullptr)
 	{
 		FMemory::Memzero(ControllerStates, sizeof(ControllerStates));
-		for( int32 ControllerPairIndex = 0; ControllerPairIndex < MaxUnrealControllers; ++ControllerPairIndex )
-		{
-			for (int32 HandIndex = 0; HandIndex < CONTROLLERS_PER_PLAYER; ++HandIndex)
-			{
-				ControllerStates[ UnrealControllerIdToControllerIndex( ControllerPairIndex, (EControllerHand)HandIndex ) ].Hand = (EControllerHand)HandIndex;
-			}
-		}
 
 		for (int32 i=0; i < vr::k_unMaxTrackedDeviceCount; ++i)
 		{
@@ -165,6 +167,7 @@ public:
 
 					DeviceToControllerMap[DeviceIndex] = FMath::FloorToInt(NumControllersMapped / CONTROLLERS_PER_PLAYER);
 					ControllerToDeviceMap[NumControllersMapped] = DeviceIndex;
+					ControllerStates[DeviceIndex].Hand = (EControllerHand)(NumControllersMapped % CONTROLLERS_PER_PLAYER);
 					++NumControllersMapped;
 
 					// update the SteamVR plugin with the new mapping
@@ -184,6 +187,15 @@ public:
 				// get the controller index for this device
 				int32 ControllerIndex = DeviceToControllerMap[DeviceIndex];
 				FControllerState& ControllerState = ControllerStates[ DeviceIndex ];
+				EControllerHand HandToUse = ControllerState.Hand;
+
+				// check to see if we need to swap input hands for debugging
+				static const auto CVar = IConsoleManager::Get().FindTConsoleVariableDataInt(TEXT("vr.SwapMotionControllerInput"));
+				bool bSwapHandInput= (CVar->GetValueOnGameThread() != 0) ? true : false;
+				if(bSwapHandInput)
+				{
+					HandToUse = (HandToUse == EControllerHand::Left) ? EControllerHand::Right : EControllerHand::Left; 
+				}
 
 				if (VRSystem->GetControllerState(DeviceIndex, &VRControllerState))
 				{
@@ -223,21 +235,21 @@ public:
 
 						if ( ControllerState.TouchPadXAnalog != VRControllerState.rAxis[TOUCHPAD_AXIS].x)
 						{
-							const FGamepadKeyNames::Type AxisButton = (ControllerState.Hand == EControllerHand::Left) ? FGamepadKeyNames::MotionController_Left_Thumbstick_X : FGamepadKeyNames::MotionController_Right_Thumbstick_X;
+							const FGamepadKeyNames::Type AxisButton = (HandToUse == EControllerHand::Left) ? FGamepadKeyNames::MotionController_Left_Thumbstick_X : FGamepadKeyNames::MotionController_Right_Thumbstick_X;
 							MessageHandler->OnControllerAnalog(AxisButton, ControllerIndex, VRControllerState.rAxis[TOUCHPAD_AXIS].x);
 							ControllerState.TouchPadXAnalog = VRControllerState.rAxis[TOUCHPAD_AXIS].x;
 						}
 
 						if ( ControllerState.TouchPadYAnalog != VRControllerState.rAxis[TOUCHPAD_AXIS].y)
 						{
-							const FGamepadKeyNames::Type AxisButton = (ControllerState.Hand == EControllerHand::Left) ? FGamepadKeyNames::MotionController_Left_Thumbstick_Y : FGamepadKeyNames::MotionController_Right_Thumbstick_Y;
+							const FGamepadKeyNames::Type AxisButton = (HandToUse == EControllerHand::Left) ? FGamepadKeyNames::MotionController_Left_Thumbstick_Y : FGamepadKeyNames::MotionController_Right_Thumbstick_Y;
 							MessageHandler->OnControllerAnalog(AxisButton, ControllerIndex, VRControllerState.rAxis[TOUCHPAD_AXIS].y);
 							ControllerState.TouchPadYAnalog = VRControllerState.rAxis[TOUCHPAD_AXIS].y;
 						}
 
 						if ( ControllerState.TriggerAnalog != VRControllerState.rAxis[TRIGGER_AXIS].x)
 						{
-							const FGamepadKeyNames::Type AxisButton = (ControllerState.Hand == EControllerHand::Left) ? FGamepadKeyNames::MotionController_Left_TriggerAxis : FGamepadKeyNames::MotionController_Right_TriggerAxis;
+							const FGamepadKeyNames::Type AxisButton = (HandToUse == EControllerHand::Left) ? FGamepadKeyNames::MotionController_Left_TriggerAxis : FGamepadKeyNames::MotionController_Right_TriggerAxis;
 							MessageHandler->OnControllerAnalog(AxisButton, ControllerIndex, VRControllerState.rAxis[TRIGGER_AXIS].x);
 							ControllerState.TriggerAnalog = VRControllerState.rAxis[TRIGGER_AXIS].x;
 						}
@@ -249,11 +261,11 @@ public:
 							{
 								if (CurrentStates[ButtonIndex])
 								{
-									MessageHandler->OnControllerButtonPressed( Buttons[ (int32)ControllerState.Hand ][ ButtonIndex ], ControllerIndex, false );
+									MessageHandler->OnControllerButtonPressed( Buttons[ (int32)HandToUse ][ ButtonIndex ], ControllerIndex, false );
 								}
 								else
 								{
-									MessageHandler->OnControllerButtonReleased( Buttons[ (int32)ControllerState.Hand ][ ButtonIndex ], ControllerIndex, false );
+									MessageHandler->OnControllerButtonReleased( Buttons[ (int32)HandToUse ][ ButtonIndex ], ControllerIndex, false );
 								}
 
 								if (CurrentStates[ButtonIndex] != 0)
@@ -275,7 +287,7 @@ public:
 				{
 					if ( ControllerState.ButtonStates[ButtonIndex] != 0 && ControllerState.NextRepeatTime[ButtonIndex] <= CurrentTime)
 					{
-						MessageHandler->OnControllerButtonPressed( Buttons[ (int32)ControllerState.Hand ][ ButtonIndex ], ControllerIndex, true );
+						MessageHandler->OnControllerButtonPressed( Buttons[ (int32)HandToUse ][ ButtonIndex ], ControllerIndex, true );
 
 						// set the button's NextRepeatTime to the ButtonRepeatDelay
 						ControllerState.NextRepeatTime[ButtonIndex] = CurrentTime + ButtonRepeatDelay;
