Fixed SteamVR public module interface class not hooked up to actual module class (integrate)

diff --git a/Engine/Plugins/Runtime/Steam/SteamVR/Source/SteamVRController/Private/SteamVRController.cpp b/Engine/Plugins/Runtime/Steam/SteamVR/Source/SteamVRController/Private/SteamVRController.cpp
--- a/Engine/Plugins/Runtime/Steam/SteamVR/Source/SteamVRController/Private/SteamVRController.cpp
+++ b/Engine/Plugins/Runtime/Steam/SteamVR/Source/SteamVRController/Private/SteamVRController.cpp
@@ -1,6 +1,7 @@
 // Copyright 1998-2015 Epic Games, Inc. All Rights Reserved.
 
 #include "SteamVRControllerPrivatePCH.h"
+#include "ISteamVRControllerPlugin.h"
 #include "IMotionController.h"
 #include "../../SteamVR/Private/SteamVRHMD.h"
 
@@ -468,7 +469,7 @@ private:
 };
 
 
-class FSteamVRControllerPlugin : public IInputDeviceModule
+class FSteamVRControllerPlugin : public ISteamVRControllerPlugin
 {
 	virtual TSharedPtr< class IInputDevice > CreateInputDevice(const TSharedRef< FGenericApplicationMessageHandler >& InMessageHandler) override
 	{
