SteamVR fix to use the appropriate viewport size that matches the HMD display

diff --git a/Engine/Plugins/Runtime/Steam/SteamVR/Source/SteamVR/Private/SteamVRHMD.cpp b/Engine/Plugins/Runtime/Steam/SteamVR/Source/SteamVR/Private/SteamVRHMD.cpp
--- a/Engine/Plugins/Runtime/Steam/SteamVR/Source/SteamVR/Private/SteamVRHMD.cpp
+++ b/Engine/Plugins/Runtime/Steam/SteamVR/Source/SteamVR/Private/SteamVRHMD.cpp
@@ -5,9 +5,31 @@
 
 #include "RendererPrivate.h"
 #include "ScenePrivate.h"
+#include "SceneViewport.h"
 #include "PostProcess/PostProcessHMD.h"
 #include "Classes/SteamVRFunctionLibrary.h"
 
+#if WITH_EDITOR
+#include "Editor/UnrealEd/Classes/Editor/EditorEngine.h"
+#endif
+
+/** Helper function for acquiring the appropriate FSceneViewport */
+FSceneViewport* FindSceneViewport()
+{
+	if (!GIsEditor)
+	{
+		UGameEngine* GameEngine = Cast<UGameEngine>(GEngine);
+		return GameEngine->SceneViewport.Get();
+	}
+#if WITH_EDITOR
+	else
+	{
+		UEditorEngine* EditorEngine = Cast<UEditorEngine>(GEngine);
+		return (FSceneViewport*)(EditorEngine->GetPIEViewport());
+	}
+#endif
+	return nullptr;
+}
 
 //---------------------------------------------------
 // SteamVR Plugin Implementation
@@ -608,6 +630,16 @@ bool FSteamVRHMD::EnableStereo(bool bStereo)
 
 	FSystemResolution::RequestResolutionChange(1280, 720, (bStereo) ? EWindowMode::WindowedMirror : EWindowMode::Windowed);
 
+	// Set the viewport to match that of the HMD display
+	FSceneViewport* SceneVP = FindSceneViewport();
+	if (VRSystem && SceneVP)
+	{
+		int32 PosX, PosY;
+		uint32 Width, Height;
+		VRSystem->GetWindowBounds(&PosX, &PosY, &Width, &Height);
+		SceneVP->SetViewportSize(Width, Height);
+	}
+
 	// Uncap fps to enable FPS higher than 62
 	GEngine->bForceDisableFrameRateSmoothing = bStereo;
 	
diff --git a/Engine/Plugins/Runtime/Steam/SteamVR/Source/SteamVR/SteamVR.Build.cs b/Engine/Plugins/Runtime/Steam/SteamVR/Source/SteamVR/SteamVR.Build.cs
--- a/Engine/Plugins/Runtime/Steam/SteamVR/Source/SteamVR/SteamVR.Build.cs
+++ b/Engine/Plugins/Runtime/Steam/SteamVR/Source/SteamVR/SteamVR.Build.cs
@@ -29,7 +29,12 @@ namespace UnrealBuildTool.Rules
 					"SlateCore"
 				}
 				);
-			
+            
+            if (UEBuildConfiguration.bBuildEditor == true)
+            {
+                PrivateDependencyModuleNames.Add("UnrealEd");
+            }
+
             if (Target.Platform == UnrealTargetPlatform.Win32 || Target.Platform == UnrealTargetPlatform.Win64)
             {
 				AddThirdPartyPrivateStaticDependencies(Target, "OpenVR");
