Fixing SteamVR late-update

diff --git a/Engine/Plugins/Runtime/Steam/SteamVR/Source/SteamVR/Private/SteamVRHMD.cpp b/Engine/Plugins/Runtime/Steam/SteamVR/Source/SteamVR/Private/SteamVRHMD.cpp
--- a/Engine/Plugins/Runtime/Steam/SteamVR/Source/SteamVR/Private/SteamVRHMD.cpp
+++ b/Engine/Plugins/Runtime/Steam/SteamVR/Source/SteamVR/Private/SteamVRHMD.cpp
@@ -891,13 +891,13 @@ void FSteamVRHMD::SetupView(FSceneViewFamily& InViewFamily, FSceneView& InView)
 
 void FSteamVRHMD::PreRenderView_RenderThread(FRHICommandListImmediate& RHICmdList, FSceneView& View)
 {
-	//check(IsInRenderingThread());
+	check(IsInRenderingThread());
 
-	//// The last view location used to set the view will be in BaseHmdOrientation.  We need to calculate the delta from that, so that
-	//// cameras that rely on game objects (e.g. other components) for their positions don't need to be updated on the render thread.
-	//const FQuat DeltaOrient = View.BaseHmdOrientation.Inverse() * TrackingFrame.DeviceOrientation[vr::k_unTrackedDeviceIndex_Hmd];
-	//View.ViewRotation = FRotator(View.ViewRotation.Quaternion() * DeltaOrient);
- //	View.UpdateViewMatrix();
+	// The last view location used to set the view will be in BaseHmdOrientation.  We need to calculate the delta from that, so that
+	// cameras that rely on game objects (e.g. other components) for their positions don't need to be updated on the render thread.
+	const FQuat DeltaOrient = View.BaseHmdOrientation.Inverse() * TrackingFrame.DeviceOrientation[vr::k_unTrackedDeviceIndex_Hmd];
+	View.ViewRotation = FRotator(View.ViewRotation.Quaternion() * DeltaOrient);
+ 	View.UpdateViewMatrix();
 }
 
 void FSteamVRHMD::PreRenderViewFamily_RenderThread(FRHICommandListImmediate& RHICmdList, FSceneViewFamily& ViewFamily)
