Separating SteamVR tracking frames for game and render threads. Each thread samples it's own tracking frame at the beginning of it's own frame's execution.

diff --git a/Engine/Plugins/Runtime/Steam/SteamVR/Source/SteamVR/Private/SteamVRHMD.cpp b/Engine/Plugins/Runtime/Steam/SteamVR/Source/SteamVR/Private/SteamVRHMD.cpp
--- a/Engine/Plugins/Runtime/Steam/SteamVR/Source/SteamVR/Private/SteamVRHMD.cpp
+++ b/Engine/Plugins/Runtime/Steam/SteamVR/Source/SteamVR/Private/SteamVRHMD.cpp
@@ -205,6 +205,7 @@ void FSteamVRHMD::GetCurrentPose(FQuat& CurrentOrientation, FVector& CurrentPosi
 
 	check(DeviceId >= 0 && DeviceId < vr::k_unMaxTrackedDeviceCount);
 
+	FTrackingFrame& TrackingFrame = const_cast<FTrackingFrame&>(GetTrackingFrame());
 	if (RefreshMode != EPoseRefreshMode::None)
 	{
 		TrackingFrame.FrameNumber = GFrameNumberRenderThread;
@@ -271,6 +272,7 @@ bool FSteamVRHMD::IsInsideBounds()
 {
 	if (VRChaperone)
 	{
+		const FTrackingFrame& TrackingFrame = GetTrackingFrame();
 		vr::HmdMatrix34_t VRPose = TrackingFrame.RawPoses[vr::k_unTrackedDeviceIndex_Hmd];
 		FMatrix Pose = ToFMatrix(VRPose);
 		
@@ -405,6 +407,7 @@ void FSteamVRHMD::GetCurrentOrientationAndPosition(FQuat& CurrentOrientation, FV
 
 float FSteamVRHMD::GetWorldToMetersScale() const
 {
+	const FTrackingFrame& TrackingFrame = GetTrackingFrame();
 	return TrackingFrame.bPoseIsValid ? TrackingFrame.WorldToMetersScale : 100.0f;
 }
 
@@ -430,6 +433,7 @@ void FSteamVRHMD::GetTrackedDeviceIds(ESteamVRTrackedDeviceType DeviceType, TArr
 {
 	TrackedIds.Empty();
 
+	const FTrackingFrame& TrackingFrame = GetTrackingFrame();
 	for (uint32 i = 0; i < vr::k_unMaxTrackedDeviceCount; ++i)
 	{
 		// Add only devices with a currently valid tracked pose, and exclude the HMD
@@ -446,6 +450,7 @@ bool FSteamVRHMD::GetTrackedObjectOrientationAndPosition(uint32 DeviceId, FQuat&
 {
 	bool bHasValidPose = false;
 
+	const FTrackingFrame& TrackingFrame = GetTrackingFrame();
 	if (DeviceId < vr::k_unMaxTrackedDeviceCount)
 	{
 		CurrentOrientation = TrackingFrame.DeviceOrientation[DeviceId];
@@ -461,6 +466,7 @@ ETrackingStatus FSteamVRHMD::GetControllerTrackingStatus(uint32 DeviceId) const
 {
 	ETrackingStatus TrackingStatus = ETrackingStatus::NotTracked;
 
+	const FTrackingFrame& TrackingFrame = GetTrackingFrame();
 	if (DeviceId < vr::k_unMaxTrackedDeviceCount && TrackingFrame.bPoseIsValid[DeviceId] && TrackingFrame.bDeviceIsConnected[DeviceId])
 	{
 		TrackingStatus = ETrackingStatus::Tracked;
@@ -695,6 +701,8 @@ void FSteamVRHMD::ResetOrientationAndPosition(float yaw)
 
 void FSteamVRHMD::ResetOrientation(float Yaw)
 {
+	const FTrackingFrame& TrackingFrame = GetTrackingFrame();
+
 	FRotator ViewRotation;
 	ViewRotation = FRotator(TrackingFrame.DeviceOrientation[vr::k_unTrackedDeviceIndex_Hmd]);
 	ViewRotation.Pitch = 0;
@@ -712,6 +720,7 @@ void FSteamVRHMD::ResetOrientation(float Yaw)
 }
 void FSteamVRHMD::ResetPosition()
 {
+	const FTrackingFrame& TrackingFrame = GetTrackingFrame();
 	FMatrix Pose = ToFMatrix(TrackingFrame.RawPoses[vr::k_unTrackedDeviceIndex_Hmd]);
 	BaseOffset = FVector(-Pose.M[3][2], Pose.M[3][0], Pose.M[3][1]);
 }
@@ -803,6 +812,7 @@ void FSteamVRHMD::CalculateStereoViewOffset(const enum EStereoscopicPass StereoP
 
 		if (!bImplicitHmdPosition)
 		{
+			const FTrackingFrame& TrackingFrame = GetTrackingFrame();
  			const FVector vHMDPosition = DeltaControlOrientation.RotateVector(TrackingFrame.DevicePosition[vr::k_unTrackedDeviceIndex_Hmd]);
 			ViewLocation += vHMDPosition;
 		}
@@ -893,6 +903,8 @@ void FSteamVRHMD::PreRenderView_RenderThread(FRHICommandListImmediate& RHICmdLis
 {
 	check(IsInRenderingThread());
 
+	const FTrackingFrame& TrackingFrame = GetTrackingFrame();
+
 	// The last view location used to set the view will be in BaseHmdOrientation.  We need to calculate the delta from that, so that
 	// cameras that rely on game objects (e.g. other components) for their positions don't need to be updated on the render thread.
 	const FQuat DeltaOrient = View.BaseHmdOrientation.Inverse() * TrackingFrame.DeviceOrientation[vr::k_unTrackedDeviceIndex_Hmd];
@@ -1313,4 +1325,16 @@ void FSteamVRHMD::SetupOcclusionMeshes()
 	}
 }
 
+const FSteamVRHMD::FTrackingFrame& FSteamVRHMD::GetTrackingFrame() const
+{
+	if (IsInRenderingThread())
+	{
+		return RenderTrackingFrame;
+	}
+	else
+	{
+		return GameTrackingFrame;
+	}
+}
+
 #endif //STEAMVR_SUPPORTED_PLATFORMS
diff --git a/Engine/Plugins/Runtime/Steam/SteamVR/Source/SteamVR/Private/SteamVRHMD.h b/Engine/Plugins/Runtime/Steam/SteamVR/Source/SteamVR/Private/SteamVRHMD.h
--- a/Engine/Plugins/Runtime/Steam/SteamVR/Source/SteamVR/Private/SteamVRHMD.h
+++ b/Engine/Plugins/Runtime/Steam/SteamVR/Source/SteamVR/Private/SteamVRHMD.h
@@ -308,7 +308,10 @@ private:
 			FMemory::Memzero(RawPoses, MaxDevices * sizeof(vr::HmdMatrix34_t));
 		}
  	};
-	FTrackingFrame TrackingFrame;
+	FTrackingFrame GameTrackingFrame;
+	FTrackingFrame RenderTrackingFrame;
+
+	const FTrackingFrame& GetTrackingFrame() const;
 
 	/** Coverts a SteamVR-space vector to an Unreal-space vector.  Does not handle scaling, only axes conversion */
 	FORCEINLINE static FVector CONVERT_STEAMVECTOR_TO_FVECTOR(const vr::HmdVector3_t InVector)
