Fix for WinXP build break in mauled FSteamVRHMD plugin

diff --git a/Engine/Plugins/Runtime/Steam/SteamVR/Source/SteamVR/Private/SteamVRHMD.cpp b/Engine/Plugins/Runtime/Steam/SteamVR/Source/SteamVR/Private/SteamVRHMD.cpp
--- a/Engine/Plugins/Runtime/Steam/SteamVR/Source/SteamVR/Private/SteamVRHMD.cpp
+++ b/Engine/Plugins/Runtime/Steam/SteamVR/Source/SteamVR/Private/SteamVRHMD.cpp
@@ -64,7 +64,10 @@ class FSteamVRPlugin : public ISteamVRPlugin
 
 public:
 	FSteamVRPlugin::FSteamVRPlugin()
-#if STEAMVR_SUPPORTED_PLATFORMS
+#if !STEAMVR_SUPPORTED_PLATFORMS
+	{
+	}
+#else //STEAMVR_SUPPORTED_PLATFORMS
 		: VRSystem(nullptr)
 	{
         LoadOpenVRModule();
