Fixed a bug in FSteamVRController::IsGamepadAttached not properly checking for right hand tracking status.

diff --git a/Engine/Plugins/Runtime/Steam/SteamVR/Source/SteamVRController/Private/SteamVRController.cpp b/Engine/Plugins/Runtime/Steam/SteamVR/Source/SteamVRController/Private/SteamVRController.cpp
--- a/Engine/Plugins/Runtime/Steam/SteamVR/Source/SteamVRController/Private/SteamVRController.cpp
+++ b/Engine/Plugins/Runtime/Steam/SteamVR/Source/SteamVRController/Private/SteamVRController.cpp
@@ -438,7 +438,7 @@ public:
 			// Only need to check for at least one player (player index 0)
 			int32 PlayerIndex = 0;
 			ETrackingStatus LeftHandTrackingStatus = GetControllerTrackingStatus(PlayerIndex, EControllerHand::Left);
-			ETrackingStatus RightHandTrackingStatus = GetControllerTrackingStatus(PlayerIndex, EControllerHand::Left);
+			ETrackingStatus RightHandTrackingStatus = GetControllerTrackingStatus(PlayerIndex, EControllerHand::Right);
 
 			return LeftHandTrackingStatus == ETrackingStatus::Tracked || RightHandTrackingStatus == ETrackingStatus::Tracked;
 		}
