Adding analytics event for Vive device initialization

diff --git a/Engine/Plugins/Runtime/Steam/SteamVR/Source/SteamVR/Private/SteamVRHMD.cpp b/Engine/Plugins/Runtime/Steam/SteamVR/Source/SteamVR/Private/SteamVRHMD.cpp
--- a/Engine/Plugins/Runtime/Steam/SteamVR/Source/SteamVR/Private/SteamVRHMD.cpp
+++ b/Engine/Plugins/Runtime/Steam/SteamVR/Source/SteamVR/Private/SteamVRHMD.cpp
@@ -11,6 +11,9 @@
 
 #include "SteamVRMeshAssets.h"
 
+#include "EngineAnalytics.h"
+#include "Runtime/Analytics/Analytics/Public/Interfaces/IAnalyticsProvider.h"
+
 #if WITH_EDITOR
 #include "Editor/UnrealEd/Classes/Editor/EditorEngine.h"
 #endif
@@ -520,6 +523,41 @@ EHMDTrackingOrigin::Type FSteamVRHMD::GetTrackingOrigin()
 	return EHMDTrackingOrigin::Floor;
 }
 
+
+void FSteamVRHMD::RecordAnalytics()
+{
+	if (FEngineAnalytics::IsAvailable())
+	{
+		// prepare and send analytics data
+		TArray<FAnalyticsEventAttribute> EventAttributes;
+
+		IHeadMountedDisplay::MonitorInfo MonitorInfo;
+		GetHMDMonitorInfo(MonitorInfo);
+
+		uint64 MonitorId = MonitorInfo.MonitorId;
+
+		char Buf[128];
+		vr::TrackedPropertyError Error;
+		FString DeviceName = "SteamVR - Default Device Name";
+		VRSystem->GetStringTrackedDeviceProperty(vr::k_unTrackedDeviceIndex_Hmd, vr::Prop_ModelNumber_String, Buf, sizeof(Buf), &Error);
+		if (Error == vr::TrackedProp_Success)
+		{
+			DeviceName = FString(UTF8_TO_TCHAR(Buf));
+		}
+		EventAttributes.Add(FAnalyticsEventAttribute(TEXT("DeviceName"), DeviceName));
+		EventAttributes.Add(FAnalyticsEventAttribute(TEXT("DisplayDeviceName"), *MonitorInfo.MonitorName));
+		EventAttributes.Add(FAnalyticsEventAttribute(TEXT("DisplayId"), MonitorId));
+		FString MonResolution(FString::Printf(TEXT("(%d, %d)"), MonitorInfo.ResolutionX, MonitorInfo.ResolutionY));
+		EventAttributes.Add(FAnalyticsEventAttribute(TEXT("Resolution"), MonResolution));
+
+		EventAttributes.Add(FAnalyticsEventAttribute(TEXT("InterpupillaryDistance"), GetInterpupillaryDistance()));
+
+
+		FString OutStr(TEXT("Editor.VR.DeviceInitialised"));
+		FEngineAnalytics::GetProvider().RecordEvent(OutStr, EventAttributes);
+	}
+}
+
 void FSteamVRHMD::SetUnrealControllerIdAndHandToDeviceIdMap(int32 InUnrealControllerIdAndHandToDeviceIdMap[ MAX_STEAMVR_CONTROLLER_PAIRS ][ 2 ] )
 {
 	for( int32 UnrealControllerIndex = 0; UnrealControllerIndex < MAX_STEAMVR_CONTROLLER_PAIRS; ++UnrealControllerIndex )
diff --git a/Engine/Plugins/Runtime/Steam/SteamVR/Source/SteamVR/Private/SteamVRHMD.h b/Engine/Plugins/Runtime/Steam/SteamVR/Source/SteamVR/Private/SteamVRHMD.h
--- a/Engine/Plugins/Runtime/Steam/SteamVR/Source/SteamVR/Private/SteamVRHMD.h
+++ b/Engine/Plugins/Runtime/Steam/SteamVR/Source/SteamVR/Private/SteamVRHMD.h
@@ -105,6 +105,8 @@ public:
 	virtual void SetTrackingOrigin(EHMDTrackingOrigin::Type NewOrigin) override;
 	virtual EHMDTrackingOrigin::Type GetTrackingOrigin() override;
 
+	virtual void RecordAnalytics() override;
+
 	/** IStereoRendering interface */
 	virtual bool IsStereoEnabled() const override;
 	virtual bool EnableStereo(bool stereo = true) override;
