Duplicating in Main:

Change: 3273750
UE-41137 //UE4/Main: Step 'Compile Ocean (Win32/Win64)' - 2 Errors - SteamVRController.cpp
-Fixing build break for Ocean.  Maybe they are using an older compiler?

Files:
//UE4/Release-4.15/Engine/Plugins/Runtime/Steam/SteamVR/Source/SteamVRController/Private/SteamVRController.cpp#4

diff --git a/Engine/Plugins/Runtime/Steam/SteamVR/Source/SteamVRController/Private/SteamVRController.cpp b/Engine/Plugins/Runtime/Steam/SteamVR/Source/SteamVRController/Private/SteamVRController.cpp
--- a/Engine/Plugins/Runtime/Steam/SteamVR/Source/SteamVRController/Private/SteamVRController.cpp
+++ b/Engine/Plugins/Runtime/Steam/SteamVR/Source/SteamVRController/Private/SteamVRController.cpp
@@ -97,8 +97,7 @@ public:
 
 	FSteamVRController(const TSharedRef< FGenericApplicationMessageHandler >& InMessageHandler)
 		: MessageHandler(InMessageHandler),
-		  SteamVRPlugin(nullptr),
-		  UnrealControllerHandUsageCount{0}
+		  SteamVRPlugin(nullptr)
 	{
 #if STEAMVRCONTROLLER_SUPPORTED_PLATFORMS
 		FMemory::Memzero(ControllerStates, sizeof(ControllerStates));
@@ -116,6 +115,11 @@ public:
 			}
 		}
 
+		for (int32& Count : UnrealControllerHandUsageCount)
+		{
+			Count = 0;
+		}
+
 		NumControllersMapped = 0;
 
 		InitialButtonRepeatDelay = 0.2f;
