Linux: fix CIS (UE-43759).

diff --git a/Engine/Plugins/Runtime/Steam/SteamVR/Source/SteamVR/Private/SteamVRHMD.h b/Engine/Plugins/Runtime/Steam/SteamVR/Source/SteamVR/Private/SteamVRHMD.h
--- a/Engine/Plugins/Runtime/Steam/SteamVR/Source/SteamVR/Private/SteamVRHMD.h
+++ b/Engine/Plugins/Runtime/Steam/SteamVR/Source/SteamVR/Private/SteamVRHMD.h
@@ -14,9 +14,7 @@
 #include <d3d11.h>
 #include "HideWindowsPlatformTypes.h"
 #elif PLATFORM_LINUX
-#define GL_GLEXT_PROTOTYPES 1
-#include <GL/glcorearb.h>
-#include <GL/glext.h>
+#include "OpenGLDrv.h"
 #endif
 
 #if STEAMVR_SUPPORTED_PLATFORMS
diff --git a/Engine/Source/Runtime/OpenGLDrv/Private/Linux/OpenGLLinux.h b/Engine/Source/Runtime/OpenGLDrv/Private/Linux/OpenGLLinux.h
--- a/Engine/Source/Runtime/OpenGLDrv/Private/Linux/OpenGLLinux.h
+++ b/Engine/Source/Runtime/OpenGLDrv/Private/Linux/OpenGLLinux.h
@@ -17,7 +17,6 @@
 THIRD_PARTY_INCLUDES_START
 	#include <GL/glcorearb.h>
 	#include <GL/glext.h>
-	#include "SDL.h"
 THIRD_PARTY_INCLUDES_END
 
 /** List all OpenGL entry points used by Unreal. */
