Speculative fix for linux - defining templatized static member variable for the general template (not just the specializations we allow).

diff --git a/Engine/Plugins/Runtime/Steam/SteamVR/Source/SteamVR/Private/SteamVRAssetManager.cpp b/Engine/Plugins/Runtime/Steam/SteamVR/Source/SteamVR/Private/SteamVRAssetManager.cpp
--- a/Engine/Plugins/Runtime/Steam/SteamVR/Source/SteamVR/Private/SteamVRAssetManager.cpp
+++ b/Engine/Plugins/Runtime/Steam/SteamVR/Source/SteamVR/Private/SteamVRAssetManager.cpp
@@ -93,9 +93,22 @@ static vr::IVRRenderModels* SteamVRDevice_Impl::GetSteamVRModelManager()
 /* TSteamVRResource 
  *****************************************************************************/
 
+template<typename ResType>
+struct TSharedSteamVRResource
+{
+	TSharedSteamVRResource() : RefCount(0), RawResource(nullptr) {}
+	int32    RefCount;
+	ResType* RawResource;
+};
+
 template< typename ResType, typename IDType >
 struct TSteamVRResource
 {
+private:
+	typedef TSharedSteamVRResource<ResType>  FSharedSteamVRResource;
+	typedef TMap<IDType, FSharedSteamVRResource> TSharedResourceMap;
+	static  TSharedResourceMap SharedResources;
+
 public:
 	TSteamVRResource(const IDType& ResID, bool bKickOffLoad = true)
 		: ResourceId(ResID)
@@ -186,25 +199,17 @@ protected:
 
 	void FreeResource(vr::IVRRenderModels* VRModelManager);
 
-private:
-	struct FSharedSteamVRResource
-	{
-		FSharedSteamVRResource() : RefCount(0), RawResource(nullptr) {}
-		int32    RefCount;
-		ResType* RawResource;
-	};
-	typedef TMap<IDType, FSharedSteamVRResource> TSharedResourceMap;
-	static  TSharedResourceMap SharedResources;
-
 protected:
 	IDType   ResourceId;
 	ResType* RawResource;
 	bool     bLoadFailed;
 };
 
+template<typename ResType, typename IDType>
+TMap< IDType, TSharedSteamVRResource<ResType> > TSteamVRResource<ResType, IDType>::SharedResources;
+
 typedef TSteamVRResource<vr::RenderModel_t, FString> TSteamVRModel;
-template<>
-TSteamVRModel::TSharedResourceMap TSteamVRModel::SharedResources;
+typedef TSteamVRResource<vr::RenderModel_TextureMap_t, int32> TSteamVRTexture;
 
 template<>
 int32 TSteamVRModel::TickAsyncLoad_Internal(vr::IVRRenderModels* VRModelManager, vr::RenderModel_t** ResourceOut)
@@ -224,10 +229,6 @@ void TSteamVRModel::FreeResource(vr::IVRRenderModels* VRModelManager)
 #endif
 }
 
-typedef TSteamVRResource<vr::RenderModel_TextureMap_t, int32> TSteamVRTexture;
-template<>
-TSteamVRTexture::TSharedResourceMap TSteamVRTexture::SharedResources;
-
 template<>
 int32 TSteamVRTexture::TickAsyncLoad_Internal(vr::IVRRenderModels* VRModelManager, vr::RenderModel_TextureMap_t** ResourceOut)
 {
