Code Documentation warning fixes.

diff --git a/Engine/Plugins/Runtime/Steam/SteamVR/Source/SteamVR/Private/SteamVRAssetManager.cpp b/Engine/Plugins/Runtime/Steam/SteamVR/Source/SteamVR/Private/SteamVRAssetManager.cpp
--- a/Engine/Plugins/Runtime/Steam/SteamVR/Source/SteamVR/Private/SteamVRAssetManager.cpp
+++ b/Engine/Plugins/Runtime/Steam/SteamVR/Source/SteamVR/Private/SteamVRAssetManager.cpp
@@ -211,6 +211,7 @@ TMap< IDType, TSharedSteamVRResource<ResType> > TSteamVRResource<ResType, IDType
 typedef TSteamVRResource<vr::RenderModel_t, FString> TSteamVRModel;
 typedef TSteamVRResource<vr::RenderModel_TextureMap_t, int32> TSteamVRTexture;
 
+/// @cond DOXYGEN_WARNINGS
 template<>
 int32 TSteamVRModel::TickAsyncLoad_Internal(vr::IVRRenderModels* VRModelManager, vr::RenderModel_t** ResourceOut)
 {
@@ -220,6 +221,7 @@ int32 TSteamVRModel::TickAsyncLoad_Internal(vr::IVRRenderModels* VRModelManager,
 	return INDEX_NONE;
 #endif
 }
+/// @endcond
 
 template<>
 void TSteamVRModel::FreeResource(vr::IVRRenderModels* VRModelManager)
@@ -229,6 +231,7 @@ void TSteamVRModel::FreeResource(vr::IVRRenderModels* VRModelManager)
 #endif
 }
 
+/// @cond DOXYGEN_WARNINGS
 template<>
 int32 TSteamVRTexture::TickAsyncLoad_Internal(vr::IVRRenderModels* VRModelManager, vr::RenderModel_TextureMap_t** ResourceOut)
 {
@@ -238,6 +241,7 @@ int32 TSteamVRTexture::TickAsyncLoad_Internal(vr::IVRRenderModels* VRModelManage
 	return INDEX_NONE;
 #endif
 }
+/// @endcond
 
 template<>
 void TSteamVRTexture::FreeResource(vr::IVRRenderModels* VRModelManager)
diff --git a/Engine/Source/Editor/BlueprintGraph/Classes/EdGraphSchema_K2.h b/Engine/Source/Editor/BlueprintGraph/Classes/EdGraphSchema_K2.h
--- a/Engine/Source/Editor/BlueprintGraph/Classes/EdGraphSchema_K2.h
+++ b/Engine/Source/Editor/BlueprintGraph/Classes/EdGraphSchema_K2.h
@@ -19,6 +19,9 @@ class FMenuBuilder;
 class UBlueprint;
 class UK2Node;
 struct FTypesDatabase;
+class UEnum;
+class UClass;
+class UScriptStruct;
 
 /** Reference to an structure (only used in 'docked' palette) */
 USTRUCT()
@@ -773,9 +776,9 @@ public:
 	/** returns friendly signature name if possible or Removes any mangling to get the unmangled signature name of the function */
 	static FText GetFriendlySignatureName(const UFunction* Function);
 
-	static bool IsAllowableBlueprintVariableType(const class UEnum* InEnum);
-	static bool IsAllowableBlueprintVariableType(const class UClass* InClass);
-	static bool IsAllowableBlueprintVariableType(const class UScriptStruct *InStruct, bool bForInternalUse = false);
+	static bool IsAllowableBlueprintVariableType(const UEnum* InEnum);
+	static bool IsAllowableBlueprintVariableType(const UClass* InClass);
+	static bool IsAllowableBlueprintVariableType(const UScriptStruct *InStruct, bool bForInternalUse = false);
 
 	static bool IsPropertyExposedOnSpawn(const UProperty* Property);
 
diff --git a/Engine/Source/Editor/ClothPainter/Private/SClothAssetSelector.cpp b/Engine/Source/Editor/ClothPainter/Private/SClothAssetSelector.cpp
--- a/Engine/Source/Editor/ClothPainter/Private/SClothAssetSelector.cpp
+++ b/Engine/Source/Editor/ClothPainter/Private/SClothAssetSelector.cpp
@@ -834,7 +834,7 @@ TWeakObjectPtr<UClothingAsset> SClothAssetSelector::GetSelectedAsset() const
 	
 }
 
-int32 SClothAssetSelector::SClothAssetSelector::GetSelectedLod() const
+int32 SClothAssetSelector::GetSelectedLod() const
 {
 	return SelectedLod;
 }
