Fix Mac DDC warning

diff --git a/Engine/Plugins/Runtime/Steam/SteamVR/Source/SteamVR/Private/SteamVRPrivate.h b/Engine/Plugins/Runtime/Steam/SteamVR/Source/SteamVR/Private/SteamVRPrivate.h
--- a/Engine/Plugins/Runtime/Steam/SteamVR/Source/SteamVR/Private/SteamVRPrivate.h
+++ b/Engine/Plugins/Runtime/Steam/SteamVR/Source/SteamVR/Private/SteamVRPrivate.h
@@ -3,7 +3,11 @@
 #pragma once
 
 /** Name of the current OpenVR SDK version in use (matches directory name) */
+#if PLATFORM_MAC
+#define OPENVR_SDK_VER TEXT("OpenVRv1_0_6")
+#else
 #define OPENVR_SDK_VER TEXT("OpenVRv1_0_7")
+#endif
 
 // for STEAMVR_SUPPORTED_PLATFORMS, keep at top
 #include "ISteamVRPlugin.h"
