UE4 - Fix SteamVR crashing Vulkan without a headset

diff --git a/Engine/Plugins/Runtime/Steam/SteamVR/Source/SteamVR/Private/SteamVRHMD.cpp b/Engine/Plugins/Runtime/Steam/SteamVR/Source/SteamVR/Private/SteamVRHMD.cpp
--- a/Engine/Plugins/Runtime/Steam/SteamVR/Source/SteamVR/Private/SteamVRHMD.cpp
+++ b/Engine/Plugins/Runtime/Steam/SteamVR/Source/SteamVR/Private/SteamVRHMD.cpp
@@ -139,7 +139,7 @@ public:
 #else //STEAMVR_SUPPORTED_PLATFORMS
 		: VRSystem(nullptr)
 	{
-        LoadOpenVRModule();
+		LoadOpenVRModule();
 	}
 
 	virtual void StartupModule() override
@@ -148,89 +148,89 @@ public:
 	
 		LoadOpenVRModule();
 	}
-    
-    virtual void ShutdownModule() override
-    {
-        IHeadMountedDisplayModule::ShutdownModule();
-        
-        UnloadOpenVRModule();
-    }
+	
+	virtual void ShutdownModule() override
+	{
+		IHeadMountedDisplayModule::ShutdownModule();
+		
+		UnloadOpenVRModule();
+	}
 
 	virtual vr::IVRSystem* GetVRSystem() const override
 	{
 		return VRSystem;
 	}
-    
-    bool LoadOpenVRModule()
-    {
+	
+	bool LoadOpenVRModule()
+	{
 #if PLATFORM_WINDOWS
 #if PLATFORM_64BITS
-        FString RootOpenVRPath;
-        TCHAR VROverridePath[MAX_PATH];
-        FPlatformMisc::GetEnvironmentVariable(TEXT("VR_OVERRIDE"), VROverridePath, MAX_PATH);
-        
-        if (FCString::Strlen(VROverridePath) > 0)
-        {
-            RootOpenVRPath = FString::Printf(TEXT("%s\\bin\\win64\\"), VROverridePath);
-        }
-        else
-        {
-            RootOpenVRPath = FPaths::EngineDir() / FString::Printf(TEXT("Binaries/ThirdParty/OpenVR/%s/Win64/"), OPENVR_SDK_VER);
-        }
-        
-        FPlatformProcess::PushDllDirectory(*RootOpenVRPath);
-        OpenVRDLLHandle = FPlatformProcess::GetDllHandle(*(RootOpenVRPath + "openvr_api.dll"));
-        FPlatformProcess::PopDllDirectory(*RootOpenVRPath);
+		FString RootOpenVRPath;
+		TCHAR VROverridePath[MAX_PATH];
+		FPlatformMisc::GetEnvironmentVariable(TEXT("VR_OVERRIDE"), VROverridePath, MAX_PATH);
+		
+		if (FCString::Strlen(VROverridePath) > 0)
+		{
+			RootOpenVRPath = FString::Printf(TEXT("%s\\bin\\win64\\"), VROverridePath);
+		}
+		else
+		{
+			RootOpenVRPath = FPaths::EngineDir() / FString::Printf(TEXT("Binaries/ThirdParty/OpenVR/%s/Win64/"), OPENVR_SDK_VER);
+		}
+		
+		FPlatformProcess::PushDllDirectory(*RootOpenVRPath);
+		OpenVRDLLHandle = FPlatformProcess::GetDllHandle(*(RootOpenVRPath + "openvr_api.dll"));
+		FPlatformProcess::PopDllDirectory(*RootOpenVRPath);
 #else
-        FString RootOpenVRPath = FPaths::EngineDir() / FString::Printf(TEXT("Binaries/ThirdParty/OpenVR/%s/Win32/"), OPENVR_SDK_VER);
-        FPlatformProcess::PushDllDirectory(*RootOpenVRPath);
-        OpenVRDLLHandle = FPlatformProcess::GetDllHandle(*(RootOpenVRPath + "openvr_api.dll"));
-        FPlatformProcess::PopDllDirectory(*RootOpenVRPath);
+		FString RootOpenVRPath = FPaths::EngineDir() / FString::Printf(TEXT("Binaries/ThirdParty/OpenVR/%s/Win32/"), OPENVR_SDK_VER);
+		FPlatformProcess::PushDllDirectory(*RootOpenVRPath);
+		OpenVRDLLHandle = FPlatformProcess::GetDllHandle(*(RootOpenVRPath + "openvr_api.dll"));
+		FPlatformProcess::PopDllDirectory(*RootOpenVRPath);
 #endif
 #elif PLATFORM_MAC
-        FString RootOpenVRPath = FPaths::EngineDir() / FString::Printf(TEXT("Binaries/ThirdParty/OpenVR/%s/osx32/"), OPENVR_SDK_VER);
-        UE_LOG(LogHMD, Log, TEXT("Tried to load %s"), *(RootOpenVRPath + "libopenvr_api.dylib"));
-        OpenVRDLLHandle = FPlatformProcess::GetDllHandle(*(RootOpenVRPath + "libopenvr_api.dylib"));
+		FString RootOpenVRPath = FPaths::EngineDir() / FString::Printf(TEXT("Binaries/ThirdParty/OpenVR/%s/osx32/"), OPENVR_SDK_VER);
+		UE_LOG(LogHMD, Log, TEXT("Tried to load %s"), *(RootOpenVRPath + "libopenvr_api.dylib"));
+		OpenVRDLLHandle = FPlatformProcess::GetDllHandle(*(RootOpenVRPath + "libopenvr_api.dylib"));
 #elif PLATFORM_LINUX
 		FString RootOpenVRPath = FPaths::EngineDir() / FString::Printf(TEXT("Binaries/ThirdParty/OpenVR/%s/linux64/"), OPENVR_SDK_VER);
 		OpenVRDLLHandle = FPlatformProcess::GetDllHandle(*(RootOpenVRPath + "libopenvr_api.so"));
 #else
 		#error "SteamVRHMD is not supported for this platform."
 #endif	//PLATFORM_WINDOWS
-        
-        if (!OpenVRDLLHandle)
-        {
-            UE_LOG(LogHMD, Log, TEXT("Failed to load OpenVR library."));
-            return false;
-        }
-        
-        //@todo steamvr: Remove GetProcAddress() workaround once we update to Steamworks 1.33 or higher
-        FSteamVRHMD::VRIsHmdPresentFn = (pVRIsHmdPresent)FPlatformProcess::GetDllExport(OpenVRDLLHandle, TEXT("VR_IsHmdPresent"));
-        FSteamVRHMD::VRGetGenericInterfaceFn = (pVRGetGenericInterface)FPlatformProcess::GetDllExport(OpenVRDLLHandle, TEXT("VR_GetGenericInterface"));
-        
-        // Note:  If this fails to compile, it's because you merged a new OpenVR version, and didn't put in the module hacks marked with @epic in openvr.h
+		
+		if (!OpenVRDLLHandle)
+		{
+			UE_LOG(LogHMD, Log, TEXT("Failed to load OpenVR library."));
+			return false;
+		}
+		
+		//@todo steamvr: Remove GetProcAddress() workaround once we update to Steamworks 1.33 or higher
+		FSteamVRHMD::VRIsHmdPresentFn = (pVRIsHmdPresent)FPlatformProcess::GetDllExport(OpenVRDLLHandle, TEXT("VR_IsHmdPresent"));
+		FSteamVRHMD::VRGetGenericInterfaceFn = (pVRGetGenericInterface)FPlatformProcess::GetDllExport(OpenVRDLLHandle, TEXT("VR_GetGenericInterface"));
+		
+		// Note:  If this fails to compile, it's because you merged a new OpenVR version, and didn't put in the module hacks marked with @epic in openvr.h
 		vr::VR_SetGenericInterfaceCallback(FSteamVRHMD::VRGetGenericInterfaceFn);
 
-        // Verify that we've bound correctly to the DLL functions
-        if (!FSteamVRHMD::VRIsHmdPresentFn || !FSteamVRHMD::VRGetGenericInterfaceFn)
-        {
-            UE_LOG(LogHMD, Log, TEXT("Failed to GetProcAddress() on openvr_api.dll"));
-            UnloadOpenVRModule();
-            
-            return false;
-        }
-        
-        return true;
-    }
-    
-    void UnloadOpenVRModule()
-    {
-        if (OpenVRDLLHandle != nullptr)
-        {
-            FPlatformProcess::FreeDllHandle(OpenVRDLLHandle);
-            OpenVRDLLHandle = nullptr;
-        }
-    }
+		// Verify that we've bound correctly to the DLL functions
+		if (!FSteamVRHMD::VRIsHmdPresentFn || !FSteamVRHMD::VRGetGenericInterfaceFn)
+		{
+			UE_LOG(LogHMD, Log, TEXT("Failed to GetProcAddress() on openvr_api.dll"));
+			UnloadOpenVRModule();
+			
+			return false;
+		}
+		
+		return true;
+	}
+	
+	void UnloadOpenVRModule()
+	{
+		if (OpenVRDLLHandle != nullptr)
+		{
+			FPlatformProcess::FreeDllHandle(OpenVRDLLHandle);
+			OpenVRDLLHandle = nullptr;
+		}
+	}
 
 	virtual bool IsHMDConnected() override
 	{
@@ -245,40 +245,40 @@ public:
 	uint64 GetGraphicsAdapterLuid() override
 	{
 #if PLATFORM_MAC
-        
-   		id<MTLDevice> SelectedDevice = nil;
-        
+		
+		id<MTLDevice> SelectedDevice = nil;
+		
 		// @TODO  currently, for mac, GetGraphicsAdapterLuid() is used to return a device index (how the function 
 		//        "GetGraphicsAdapter" used to work), not a ID... eventually we want the HMD module to return the 
 		//        MTLDevice's registryID, but we cannot fully handle that until we drop support for 10.12
 		//  NOTE: this is why we  use -1 as a sentinel value representing "no device" (instead of 0, which is used in the LUID case)
-        {
-            // HACK:  Temporarily stand up the VRSystem to get a graphics adapter.  We're pretty sure we're going to use SteamVR if we get in here, but not guaranteed
-            vr::EVRInitError VRInitErr = vr::VRInitError_None;
-            // Attempt to initialize the VRSystem device
-            vr::IVRSystem* TempVRSystem = vr::VR_Init(&VRInitErr, vr::VRApplication_Scene);
-            if ((TempVRSystem == nullptr) || (VRInitErr != vr::VRInitError_None))
-            {
-                UE_LOG(LogHMD, Log, TEXT("Failed to initialize OpenVR with code %d"), (int32)VRInitErr);
-                return (uint64)-1;
-            }
-            
-            // Make sure that the version of the HMD we're compiled against is correct.  This will fill out the proper vtable!
-            TempVRSystem = (vr::IVRSystem*)(*FSteamVRHMD::VRGetGenericInterfaceFn)(vr::IVRSystem_Version, &VRInitErr);
-            if ((TempVRSystem == nullptr) || (VRInitErr != vr::VRInitError_None))
-            {
-                return (uint64)-1;
-            }
-            
-            TempVRSystem->GetOutputDevice((uint64_t*)((void*)&SelectedDevice), vr::TextureType_IOSurface);
-            
-            vr::VR_Shutdown();
-        }
-
-        if(SelectedDevice == nil)
-        {
-            return (uint64)-1;
-        }
+		{
+			// HACK:  Temporarily stand up the VRSystem to get a graphics adapter.  We're pretty sure we're going to use SteamVR if we get in here, but not guaranteed
+			vr::EVRInitError VRInitErr = vr::VRInitError_None;
+			// Attempt to initialize the VRSystem device
+			vr::IVRSystem* TempVRSystem = vr::VR_Init(&VRInitErr, vr::VRApplication_Scene);
+			if ((TempVRSystem == nullptr) || (VRInitErr != vr::VRInitError_None))
+			{
+				UE_LOG(LogHMD, Log, TEXT("Failed to initialize OpenVR with code %d"), (int32)VRInitErr);
+				return (uint64)-1;
+			}
+			
+			// Make sure that the version of the HMD we're compiled against is correct.  This will fill out the proper vtable!
+			TempVRSystem = (vr::IVRSystem*)(*FSteamVRHMD::VRGetGenericInterfaceFn)(vr::IVRSystem_Version, &VRInitErr);
+			if ((TempVRSystem == nullptr) || (VRInitErr != vr::VRInitError_None))
+			{
+				return (uint64)-1;
+			}
+			
+			TempVRSystem->GetOutputDevice((uint64_t*)((void*)&SelectedDevice), vr::TextureType_IOSurface);
+			
+			vr::VR_Shutdown();
+		}
+
+		if(SelectedDevice == nil)
+		{
+			return (uint64)-1;
+		}
 
 		TArray<FMacPlatformMisc::FGPUDescriptor> const& GPUs = FPlatformMisc::GetGPUDescriptors();
 		check(GPUs.Num() > 0);
@@ -319,15 +319,19 @@ public:
 
 	virtual TSharedPtr< IHeadMountedDisplayVulkanExtensions, ESPMode::ThreadSafe > GetVulkanExtensions() override
 	{
-		FSteamVRHMD* SteamVRHMD = new FSteamVRHMD( this );
-		check( SteamVRHMD->IsInitialized() );
-		return TSharedPtr< IHeadMountedDisplayVulkanExtensions, ESPMode::ThreadSafe >( SteamVRHMD );
+		FSteamVRHMD* SteamVRHMD = new FSteamVRHMD(this);
+		if (!SteamVRHMD->IsInitialized())
+		{
+			delete SteamVRHMD;
+			return nullptr;
+		}
+		return TSharedPtr< IHeadMountedDisplayVulkanExtensions, ESPMode::ThreadSafe >(SteamVRHMD);
 	}
 
 private:
 	vr::IVRSystem* VRSystem;
-    
-    void* OpenVRDLLHandle;
+	
+	void* OpenVRDLLHandle;
 #endif // STEAMVR_SUPPORTED_PLATFORMS
 };
 
@@ -358,33 +362,33 @@ pVRGetGenericInterface FSteamVRHMD::VRGetGenericInterfaceFn = nullptr;
 
 bool FSteamVRHMD::GetVulkanInstanceExtensionsRequired(TArray<const ANSICHAR*>& Out)
 {
- 	if (VRCompositor == nullptr)
- 	{
- 		UE_LOG(LogHMD, Warning, TEXT("VRCompositor is null"));
- 		return false;
- 	}
+	if (VRCompositor == nullptr)
+	{
+		UE_LOG(LogHMD, Warning, TEXT("VRCompositor is null"));
+		return false;
+	}
  
- 	static ANSICHAR* InstanceExtensionsBuf = nullptr;
+	static ANSICHAR* InstanceExtensionsBuf = nullptr;
  
- 	uint32_t BufSize = VRCompositor->GetVulkanInstanceExtensionsRequired(nullptr, 0);
- 	if (BufSize == 0)
- 	{
- 		return true; // No particular extensions required
- 	}
- 	if (InstanceExtensionsBuf != nullptr)
- 	{
- 		FMemory::Free(InstanceExtensionsBuf);
- 	}
- 	InstanceExtensionsBuf = (ANSICHAR*)FMemory::Malloc(BufSize);
- 	VRCompositor->GetVulkanInstanceExtensionsRequired(InstanceExtensionsBuf, BufSize);
+	uint32_t BufSize = VRCompositor->GetVulkanInstanceExtensionsRequired(nullptr, 0);
+	if (BufSize == 0)
+	{
+		return true; // No particular extensions required
+	}
+	if (InstanceExtensionsBuf != nullptr)
+	{
+		FMemory::Free(InstanceExtensionsBuf);
+	}
+	InstanceExtensionsBuf = (ANSICHAR*)FMemory::Malloc(BufSize);
+	VRCompositor->GetVulkanInstanceExtensionsRequired(InstanceExtensionsBuf, BufSize);
  
- 	ANSICHAR * Context = nullptr;
- 	ANSICHAR * Tok = FCStringAnsi::Strtok(InstanceExtensionsBuf, " ", &Context);
- 	while (Tok != nullptr)
- 	{
- 		Out.Add(Tok);
- 		Tok = FCStringAnsi::Strtok(nullptr, " ", &Context);
- 	}
+	ANSICHAR * Context = nullptr;
+	ANSICHAR * Tok = FCStringAnsi::Strtok(InstanceExtensionsBuf, " ", &Context);
+	while (Tok != nullptr)
+	{
+		Out.Add(Tok);
+		Tok = FCStringAnsi::Strtok(nullptr, " ", &Context);
+	}
 
 	return true;
 }
@@ -392,32 +396,32 @@ bool FSteamVRHMD::GetVulkanInstanceExtensionsRequired(TArray<const ANSICHAR*>& O
 bool FSteamVRHMD::GetVulkanDeviceExtensionsRequired(VkPhysicalDevice_T *pPhysicalDevice, TArray<const ANSICHAR*>& Out)
 {
 	if ( VRCompositor == nullptr )
- 	{
- 		UE_LOG(LogHMD, Warning, TEXT("VRCompositor is null"));
- 		return false;
- 	}
+	{
+		UE_LOG(LogHMD, Warning, TEXT("VRCompositor is null"));
+		return false;
+	}
  
- 	static ANSICHAR* DeviceExtensionsBuf = nullptr;
+	static ANSICHAR* DeviceExtensionsBuf = nullptr;
  
- 	uint32_t BufSize = VRCompositor->GetVulkanDeviceExtensionsRequired(pPhysicalDevice, nullptr, 0);
- 	if (BufSize == 0)
- 	{
- 		return true; // No particular extensions required
- 	}
- 	if (DeviceExtensionsBuf != nullptr)
- 	{
- 		FMemory::Free(DeviceExtensionsBuf);
- 	}
- 	DeviceExtensionsBuf = (ANSICHAR*)FMemory::Malloc(BufSize);
- 	VRCompositor->GetVulkanDeviceExtensionsRequired(pPhysicalDevice, DeviceExtensionsBuf, BufSize);
+	uint32_t BufSize = VRCompositor->GetVulkanDeviceExtensionsRequired(pPhysicalDevice, nullptr, 0);
+	if (BufSize == 0)
+	{
+		return true; // No particular extensions required
+	}
+	if (DeviceExtensionsBuf != nullptr)
+	{
+		FMemory::Free(DeviceExtensionsBuf);
+	}
+	DeviceExtensionsBuf = (ANSICHAR*)FMemory::Malloc(BufSize);
+	VRCompositor->GetVulkanDeviceExtensionsRequired(pPhysicalDevice, DeviceExtensionsBuf, BufSize);
  
- 	ANSICHAR * Context = nullptr;
- 	ANSICHAR * Tok = FCStringAnsi::Strtok(DeviceExtensionsBuf, " ", &Context);
- 	while (Tok != nullptr)
- 	{
- 		Out.Add(Tok);
- 		Tok = FCStringAnsi::Strtok(nullptr, " ", &Context);
- 	}
+	ANSICHAR * Context = nullptr;
+	ANSICHAR * Tok = FCStringAnsi::Strtok(DeviceExtensionsBuf, " ", &Context);
+	while (Tok != nullptr)
+	{
+		Out.Add(Tok);
+		Tok = FCStringAnsi::Strtok(nullptr, " ", &Context);
+	}
 
 	return true;
 }
@@ -1454,13 +1458,13 @@ void FSteamVRHMD::CalculateRenderTargetSize(const class FViewport& Viewport, uin
 {
 	check(IsInGameThread());
 
- 	if (!IsStereoEnabled())
- 	{
- 		return;
- 	}
-    
+	if (!IsStereoEnabled())
+	{
+		return;
+	}
+	
 	InOutSizeX = FrameSettings.RenderTargetSize.X;
-    InOutSizeY = FrameSettings.RenderTargetSize.Y;
+	InOutSizeY = FrameSettings.RenderTargetSize.Y;
 
 	check(InOutSizeX != 0 && InOutSizeY != 0);
 }
@@ -1486,7 +1490,7 @@ bool FSteamVRHMD::NeedReAllocateViewportRenderTarget(const FViewport& Viewport)
 }
 
 FSteamVRHMD::FSteamVRHMD(ISteamVRPlugin* InSteamVRPlugin) :
-    CUseAdaptivePD(TEXT("vr.SteamVR.PixelDensityAdaptive"), TEXT("SteamVR adaptive pixel density support.  0 to disable, 1 to enable"), FConsoleCommandWithWorldArgsAndOutputDeviceDelegate::CreateRaw(this, &FSteamVRHMD::AdaptivePixelDensityCommandHandler)),
+	CUseAdaptivePD(TEXT("vr.SteamVR.PixelDensityAdaptive"), TEXT("SteamVR adaptive pixel density support.  0 to disable, 1 to enable"), FConsoleCommandWithWorldArgsAndOutputDeviceDelegate::CreateRaw(this, &FSteamVRHMD::AdaptivePixelDensityCommandHandler)),
 	bHmdEnabled(true),
 	HmdWornState(EHMDWornState::Unknown),
 	bStereoDesired(false),
@@ -1793,11 +1797,11 @@ static void SetupHiddenAreaMeshes(vr::IVRSystem* const VRSystem, FHMDViewMesh Re
 void FSteamVRHMD::SetupOcclusionMeshes()
 {	
 	SetupHiddenAreaMeshes(VRSystem, HiddenAreaMeshes, vr::EHiddenAreaMeshType::k_eHiddenAreaMesh_Standard);
-    
-    if(CUseSteamVRVisibleAreaMesh.GetValueOnAnyThread() > 0)
-    {
-        SetupHiddenAreaMeshes(VRSystem, VisibleAreaMeshes, vr::EHiddenAreaMeshType::k_eHiddenAreaMesh_Inverse);
-    }
+	
+	if(CUseSteamVRVisibleAreaMesh.GetValueOnAnyThread() > 0)
+	{
+		SetupHiddenAreaMeshes(VRSystem, VisibleAreaMeshes, vr::EHiddenAreaMeshType::k_eHiddenAreaMesh_Inverse);
+	}
 }
 
 const FSteamVRHMD::FTrackingFrame& FSteamVRHMD::GetTrackingFrame() const
@@ -1930,7 +1934,7 @@ int32 FSteamVRHMD::CalculateScalabilityFactor()
 	const bool bOutputDebug = (CDebugAdaptiveOutput.GetValueOnAnyThread() > 0);
 	if(bOutputDebug)
 	{
-	    UE_LOG(LogHMD, Log, TEXT("FrameTime: %2.1f, FrameTime - 1: %2.1f, Frametime - 2: %2.1f"), CurrentFrameTime, PreviousFrameTime, TwoPreviousFrameTime);
+		UE_LOG(LogHMD, Log, TEXT("FrameTime: %2.1f, FrameTime - 1: %2.1f, Frametime - 2: %2.1f"), CurrentFrameTime, PreviousFrameTime, TwoPreviousFrameTime);
 	}
 
 	return RetVal;
@@ -1953,18 +1957,18 @@ void FSteamVRHMD::UpdateStereoRenderingParams()
 		// Update values for our PD range, in case they've been changed
 		FrameSettings.PixelDensityMin = CUseAdaptivePDMin.GetValueOnAnyThread();
 		FrameSettings.PixelDensityMax = CUseAdaptivePDMax.GetValueOnAnyThread();
-        const bool bDebugAdaptive = (CDebugAdaptiveCycle.GetValueOnAnyThread() > 0);
-
-        if(bDebugAdaptive)
-        {
-            FrameSettings.CurrentPixelDensity -= 0.005;
-            if (FrameSettings.CurrentPixelDensity < FrameSettings.PixelDensityMin)
-            {
-                FrameSettings.CurrentPixelDensity = FrameSettings.PixelDensityMax;
-            }
-        }
-        else
-        {
+		const bool bDebugAdaptive = (CDebugAdaptiveCycle.GetValueOnAnyThread() > 0);
+
+		if(bDebugAdaptive)
+		{
+			FrameSettings.CurrentPixelDensity -= 0.005;
+			if (FrameSettings.CurrentPixelDensity < FrameSettings.PixelDensityMin)
+			{
+				FrameSettings.CurrentPixelDensity = FrameSettings.PixelDensityMax;
+			}
+		}
+		else
+		{
 			// Determine if we need to scale up or down based on the most recent frames.  This will tell us if we need to move up or down a bucket
 			const int32 PerformanceDelta = CalculateScalabilityFactor();
 			CurrentAdaptiveBucket = FMath::Clamp(CurrentAdaptiveBucket + PerformanceDelta, 0, AdaptivePixelDensityBuckets.Num() - 1);
@@ -1981,13 +1985,13 @@ void FSteamVRHMD::UpdateStereoRenderingParams()
 				
 				FrameSettings.CurrentPixelDensity = AdaptivePixelDensityBuckets[CurrentAdaptiveBucket];
 			}
-        }
-        
-        const bool bOutputDebug = (CDebugAdaptiveOutput.GetValueOnAnyThread() > 0);
-        if(bOutputDebug)
-        {
-            UE_LOG(LogHMD, Log, TEXT("---> PDAdaptive: %1.2f"), FrameSettings.CurrentPixelDensity);
-        }
+		}
+		
+		const bool bOutputDebug = (CDebugAdaptiveOutput.GetValueOnAnyThread() > 0);
+		if(bOutputDebug)
+		{
+			UE_LOG(LogHMD, Log, TEXT("---> PDAdaptive: %1.2f"), FrameSettings.CurrentPixelDensity);
+		}
 	}
 	else
 	{
