Mirroring CL 4080389 from Dev-VR

Speculative fix/guard against live crash - trying to catch malformed model data. Logging helpful information to give us insight in the future.

diff --git a/Engine/Plugins/Runtime/Steam/SteamVR/Source/SteamVR/Private/SteamVRAssetManager.cpp b/Engine/Plugins/Runtime/Steam/SteamVR/Source/SteamVR/Private/SteamVRAssetManager.cpp
--- a/Engine/Plugins/Runtime/Steam/SteamVR/Source/SteamVR/Private/SteamVRAssetManager.cpp
+++ b/Engine/Plugins/Runtime/Steam/SteamVR/Source/SteamVR/Private/SteamVRAssetManager.cpp
@@ -11,6 +11,7 @@
 #include "Engine/Texture2D.h"
 #include "UObject/Package.h"
 #include "UObject/GCObject.h"
+#include "Logging/LogMacros.h"
 
 #if STEAMVR_SUPPORTED_PLATFORMS
 	#include "SteamVRHMD.h"
@@ -165,7 +166,6 @@ public:
 			}
 #endif
 			else
- 
 			{
 				bLoadFailed = true;
 			}
@@ -176,6 +176,8 @@ public:
 	operator ResType*()   { return RawResource; }
 	ResType* operator->() { return RawResource; }
 
+	IDType GetId() const  { return ResourceId;  }
+
 protected:
 	int32 TickAsyncLoad_Internal(vr::IVRRenderModels* VRModelManager, ResType** ResourceOut);
 
@@ -455,30 +457,30 @@ void FSteamVRAsyncMeshLoader::Tick(float /*DeltaTime*/)
 	{
 		FSteamVRModel& ModelResource = EnqueuedMeshes[SubMeshIndex];
 
-			vr::RenderModel_t* RenderModel = ModelResource.TickAsyncLoad();
-			if (!ModelResource.IsPending())
-			{
-				--PendingLoadCount;
+		vr::RenderModel_t* RenderModel = ModelResource.TickAsyncLoad();
+		if (!ModelResource.IsPending())
+		{
+			--PendingLoadCount;
 
-				if (!RenderModel)
-				{
-					// valid index + missing RenderModel => signifies failure
-					OnLoadComplete(SubMeshIndex);
-				}
+			if (!RenderModel)
+			{
+				// valid index + missing RenderModel => signifies failure
+				OnLoadComplete(SubMeshIndex);
+			}
 #if STEAMVR_SUPPORTED_PLATFORMS
-				// if we've already loaded and converted the texture
-				else if (ConstructedTextures.Contains(RenderModel->diffuseTextureId))
-				{
-					OnLoadComplete(SubMeshIndex);
-				}
-#endif // STEAMVR_SUPPORTED_PLATFORMS
-				else if (!EnqueueTextureLoad(SubMeshIndex, RenderModel))
-				{
-					// if we fail to load the texture, we'll have to do without it
-					OnLoadComplete(SubMeshIndex);
-				}			
+			// if we've already loaded and converted the texture
+			else if (ConstructedTextures.Contains(RenderModel->diffuseTextureId))
+			{
+				OnLoadComplete(SubMeshIndex);
 			}
+#endif // STEAMVR_SUPPORTED_PLATFORMS
+			else if (!EnqueueTextureLoad(SubMeshIndex, RenderModel))
+			{
+				// if we fail to load the texture, we'll have to do without it
+				OnLoadComplete(SubMeshIndex);
+			}			
 		}
+	}
 
 	for (int32 TexIndex = 0; TexIndex < EnqueuedTextures.Num(); ++TexIndex)
 	{
@@ -561,12 +563,26 @@ void FSteamVRAsyncMeshLoader::OnLoadComplete(int32 SubMeshIndex)
 	if (EnqueuedMeshes.IsValidIndex(SubMeshIndex))
 	{
 		FSteamVRModel& LoadedModel = EnqueuedMeshes[SubMeshIndex];
-		LoadedModel.GetRawMeshData(WorldMetersScale, RawMeshData);
 
-		
 		if (LoadedModel.IsValid())
 		{
 #if STEAMVR_SUPPORTED_PLATFORMS
+			// trying to handle an illusive crash where the loaded model data appears to be bad... 
+			// technically we can handle when there is no diffuse texture, but it may be indicative 
+			// of a larger issue (we expect all steamVR models to be textured)
+			const bool bHasMalformData = (LoadedModel->diffuseTextureId == vr::INVALID_TEXTURE_ID);
+			UE_CLOG(bHasMalformData, LogSteamVR, Warning, TEXT("Loaded what appears to be malformed model data for SteamVR model (0x%08x): \n"
+				"\t %s \n" 
+				"\t Vert count: %d \n" 
+				"\t Tri  count: %d \n" 
+			"Treating as a load failure (no model will be spawned)!"), (vr::RenderModel_t*)LoadedModel, *LoadedModel.GetId(), LoadedModel->unVertexCount, LoadedModel->unTriangleCount);
+
+			if (!bHasMalformData)
+			{
+				LoadedModel.GetRawMeshData(WorldMetersScale, RawMeshData);
+			}
+			// else, skip polling mesh data as there may be a crash with vert/index buffer count mismatch
+
 			UTexture2D** CachedTexturePtr = ConstructedTextures.Find(LoadedModel->diffuseTextureId);
 			if (CachedTexturePtr)
 			{
@@ -804,6 +820,7 @@ UPrimitiveComponent* FSteamVRAssetManager::CreateRenderComponent(const int32 Dev
 			
 			FAsyncLoadData CallbackPayload;
 			CallbackPayload.ComponentPtr = ProceduralMesh;
+			CallbackPayload.LoadedModelName = ModelName;
 
 			AssignedMeshLoader.Pin()->OnSubMeshLoaded().AddRaw(this, &FSteamVRAssetManager::OnMeshLoaded, CallbackPayload);
 			AssignedMeshLoader.Pin()->OnLoadComplete().AddRaw(this, &FSteamVRAssetManager::OnComponentLoadComplete, CallbackPayload.ComponentPtr, OnLoadComplete);
@@ -847,6 +864,10 @@ void FSteamVRAssetManager::OnMeshLoaded(int32 SubMeshIndex, const FSteamVRMeshDa
 			}
 		}
 	}
+	else
+	{
+		UE_CLOG(MeshData.VertPositions.Num() <= 0, LogSteamVR, Warning, TEXT("Loaded empty sub-mesh for SteamVR device model: '%s'"), *LoadData.LoadedModelName);
+	}
 }
 
 void FSteamVRAssetManager::OnComponentLoadComplete(TWeakObjectPtr<UProceduralMeshComponent> ComponentPtr, FXRComponentLoadComplete LoadCompleteCallback)
diff --git a/Engine/Plugins/Runtime/Steam/SteamVR/Source/SteamVR/Private/SteamVRAssetManager.h b/Engine/Plugins/Runtime/Steam/SteamVR/Source/SteamVR/Private/SteamVRAssetManager.h
--- a/Engine/Plugins/Runtime/Steam/SteamVR/Source/SteamVR/Private/SteamVRAssetManager.h
+++ b/Engine/Plugins/Runtime/Steam/SteamVR/Source/SteamVR/Private/SteamVRAssetManager.h
@@ -31,6 +31,7 @@ protected:
 	struct FAsyncLoadData
 	{
 		TWeakObjectPtr<UProceduralMeshComponent> ComponentPtr;
+		FString LoadedModelName;
 	};
 	void OnMeshLoaded(int32 SubMeshIndex, const FSteamVRMeshData& MeshData, UTexture2D* DiffuseTex, FAsyncLoadData LoadData);
 	void OnComponentLoadComplete(TWeakObjectPtr<UProceduralMeshComponent> ComponentPtr, FXRComponentLoadComplete LoadCompleteCallback);
diff --git a/Engine/Plugins/Runtime/Steam/SteamVR/Source/SteamVR/Private/SteamVRHMD.cpp b/Engine/Plugins/Runtime/Steam/SteamVR/Source/SteamVR/Private/SteamVRHMD.cpp
--- a/Engine/Plugins/Runtime/Steam/SteamVR/Source/SteamVR/Private/SteamVRHMD.cpp
+++ b/Engine/Plugins/Runtime/Steam/SteamVR/Source/SteamVR/Private/SteamVRHMD.cpp
@@ -36,6 +36,8 @@
 #include <Metal/Metal.h>
 #endif
 
+DEFINE_LOG_CATEGORY(LogSteamVR);
+
 // Visibility mesh
 static TAutoConsoleVariable<int32> CUseSteamVRVisibleAreaMesh(TEXT("vr.SteamVR.UseVisibleAreaMesh"), 1, TEXT("If non-zero, SteamVR will use the visible area mesh in addition to the hidden area mesh optimization.  This may be problematic on beta versions of platforms."));
 
diff --git a/Engine/Plugins/Runtime/Steam/SteamVR/Source/SteamVR/Public/ISteamVRPlugin.h b/Engine/Plugins/Runtime/Steam/SteamVR/Source/SteamVR/Public/ISteamVRPlugin.h
--- a/Engine/Plugins/Runtime/Steam/SteamVR/Source/SteamVR/Public/ISteamVRPlugin.h
+++ b/Engine/Plugins/Runtime/Steam/SteamVR/Source/SteamVR/Public/ISteamVRPlugin.h
@@ -7,6 +7,7 @@
 #include "Modules/ModuleManager.h"
 #include "IHeadMountedDisplayModule.h"
 #include "IHeadMountedDisplay.h"
+#include "Logging/LogMacros.h"
 #if STEAMVR_SUPPORTED_PLATFORMS
 #include "openvr.h"
 #endif // STEAMVR_SUPPORTED_PLATFORMS
@@ -14,6 +15,8 @@
 /** Up to 8 motion controller devices supported (two VR motion controllers per Unreal controller, one for either the left or right hand.) */
 #define MAX_STEAMVR_CONTROLLER_PAIRS 4
 
+DECLARE_LOG_CATEGORY_EXTERN(LogSteamVR, Log, All);
+
 /**
  * The public interface to this module.  In most cases, this interface is only public to sibling modules 
  * within this plugin.
