Fix CIS warnings compiling UE4Editor, UE4Game and sample editors on Mac. There is no VulkanRHI/Private/Mac directory.

diff --git a/Engine/Plugins/Runtime/Steam/SteamVR/Source/SteamVR/SteamVR.Build.cs b/Engine/Plugins/Runtime/Steam/SteamVR/Source/SteamVR/SteamVR.Build.cs
--- a/Engine/Plugins/Runtime/Steam/SteamVR/Source/SteamVR/SteamVR.Build.cs
+++ b/Engine/Plugins/Runtime/Steam/SteamVR/Source/SteamVR/SteamVR.Build.cs
@@ -19,7 +19,7 @@ namespace UnrealBuildTool.Rules
 			{
 				PrivateIncludePaths.Add("../../../../../Source/Runtime/VulkanRHI/Private/Windows");
 			}
-			else
+			else if(Target.Platform != UnrealTargetPlatform.Mac)
 			{
 				PrivateIncludePaths.Add("../../../../../Source/Runtime/VulkanRHI/Private/" + Target.Platform);
 			}
