Mirroring CL 4081737 from Dev-VR

Fix SteamVR forcing full screen regardless of user settings.

diff --git a/Engine/Plugins/Runtime/Steam/SteamVR/Source/SteamVR/Private/SteamVRHMD.cpp b/Engine/Plugins/Runtime/Steam/SteamVR/Source/SteamVR/Private/SteamVRHMD.cpp
--- a/Engine/Plugins/Runtime/Steam/SteamVR/Source/SteamVR/Private/SteamVRHMD.cpp
+++ b/Engine/Plugins/Runtime/Steam/SteamVR/Source/SteamVR/Private/SteamVRHMD.cpp
@@ -1262,17 +1262,6 @@ bool FSteamVRHMD::EnableStereo(bool bStereo)
 		TSharedPtr<SWindow> Window = SceneVP->FindWindow();
 		if (Window.IsValid() && SceneVP->GetViewportWidget().IsValid())
 		{
-			int32 ResX = 2160;
-			int32 ResY = 1200;
-
-			MonitorInfo MonitorDesc;
-			if (GetHMDMonitorInfo(MonitorDesc))
-			{
-				ResX = MonitorDesc.ResolutionX;
-				ResY = MonitorDesc.ResolutionY;
-			}
-			FSystemResolution::RequestResolutionChange(ResX, ResY, EWindowMode::WindowedFullscreen);
-
 			if( bStereo )
 			{
 				int32 PosX, PosY;
