Don't update the overlay position for world-locked overlays on every tick

PR #5280: Fixes SteamVR Stereo layer rotations (Contributed by [REDACTED_PERSON])

diff --git a/Engine/Plugins/Runtime/Steam/SteamVR/Source/SteamVR/Private/SteamVRStereoLayers.cpp b/Engine/Plugins/Runtime/Steam/SteamVR/Source/SteamVR/Private/SteamVRStereoLayers.cpp
--- a/Engine/Plugins/Runtime/Steam/SteamVR/Source/SteamVR/Private/SteamVRStereoLayers.cpp
+++ b/Engine/Plugins/Runtime/Steam/SteamVR/Source/SteamVR/Private/SteamVRStereoLayers.cpp
@@ -26,8 +26,8 @@ static vr::EVROverlayError sOvrError;
 //=============================================================================
 static void TransformToSteamSpace(const FTransform& In, vr::HmdMatrix34_t& Out, float WorldToMeterScale)
 {
-	const FRotator InRot = In.Rotator();
-	FRotator OutRot(InRot.Yaw, -InRot.Roll, -InRot.Pitch);
+	const FQuat InRot = In.GetRotation();
+	FQuat OutRot(InRot.Y, InRot.Z, -InRot.X, -InRot.W);
 
 	const FVector InPos = In.GetTranslation();
 	FVector OutPos(InPos.Y, InPos.Z, -InPos.X);
@@ -179,14 +179,28 @@ void FSteamVRHMD::UpdateLayer(struct FSteamVRLayer& Layer, uint32 LayerId, bool
 
 	// Shift layer priority values up by -INT32_MIN, as SteamVR uses unsigned integers for the layer order where UE uses signed integers.
 	// This will preserve the correct order between layers with negative and positive priorities
-	OVR_VERIFY(VROverlay->SetOverlaySortOrder(Layer.OverlayHandle, Layer.LayerDesc.Priority-INT32_MIN)); 
+	OVR_VERIFY(VROverlay->SetOverlaySortOrder(Layer.OverlayHandle, Layer.LayerDesc.Priority-INT32_MIN));
+
+	FQuat AdjustedPlayerOrientation = BaseOrientation.Inverse() * PlayerOrientation;
+	AdjustedPlayerOrientation.Normalize();
+
+	check(XRCamera.IsValid());
+	FVector AdjustedPlayerLocation = PlayerLocation;
+	if (XRCamera->GetUseImplicitHMDPosition())
+	{
+		AdjustedPlayerLocation -= BaseOrientation.Inverse().RotateVector(GameTrackingFrame.DevicePosition[IXRTrackingSystem::HMDDeviceId]);
+	}
+
+	FTransform InvWorldTransform = FTransform(AdjustedPlayerOrientation, AdjustedPlayerLocation).Inverse();
 
 	// Transform
 	switch (Layer.LayerDesc.PositionType)
 	{
 	case ELayerType::WorldLocked:
 	{
-		// World locked layer positions are updated every frame.
+		vr::HmdMatrix34_t HmdTransform;
+		TransformToSteamSpace(Layer.LayerDesc.Transform * InvWorldTransform, HmdTransform, WorldToMeterScale);
+		OVR_VERIFY(VROverlay->SetOverlayTransformAbsolute(Layer.OverlayHandle, VRCompositor->GetTrackingSpace(), &HmdTransform));
 		break;
 	}
 	case ELayerType::TrackerLocked:
@@ -201,7 +215,7 @@ void FSteamVRHMD::UpdateLayer(struct FSteamVRLayer& Layer, uint32 LayerId, bool
 	{
 		vr::HmdMatrix34_t HmdTransform;
 		TransformToSteamSpace(Layer.LayerDesc.Transform, HmdTransform, WorldToMeterScale);
-		OVR_VERIFY(VROverlay->SetOverlayTransformTrackedDeviceRelative(Layer.OverlayHandle, 0, &HmdTransform));
+		OVR_VERIFY(VROverlay->SetOverlayTransformTrackedDeviceRelative(Layer.OverlayHandle, vr::TrackingUniverseSeated, &HmdTransform));
 	}
 	};
 }
@@ -225,37 +239,11 @@ void FSteamVRHMD::UpdateStereoLayers_RenderThread()
 	typedef TTuple<vr::VROverlayHandle_t, int32, bool /* bIsFaceLocked */> LayerPriorityInfo;
 	TArray<LayerPriorityInfo> LayerPriorities;
 
-	const float WorldToMeterScale = GetWorldToMetersScale();
-	check(WorldToMeterScale > 0.f);
-	FQuat AdjustedPlayerOrientation = BaseOrientation.Inverse() * PlayerOrientation;
-	AdjustedPlayerOrientation.Normalize();
-
-	check(XRCamera.IsValid());
-	FVector AdjustedPlayerLocation = PlayerLocation;
-	if (XRCamera->GetUseImplicitHMDPosition())
-	{
-		FQuat DeviceOrientation; // Unused
-		FVector DevicePosition;
-		GetCurrentPose(IXRTrackingSystem::HMDDeviceId, DeviceOrientation, DevicePosition);
-		AdjustedPlayerLocation -= BaseOrientation.Inverse().RotateVector(DevicePosition);
-	}
-
-	FTransform InvWorldTransform = FTransform(AdjustedPlayerOrientation, AdjustedPlayerLocation).Inverse();
-
-
 	// We have loop through all layers every frame, in case we have world locked layers or continuously updated textures.
 	ForEachLayer([&](uint32 /* unused */, FSteamVRLayer& Layer)
 	{
 		if (Layer.OverlayHandle != vr::k_ulOverlayHandleInvalid)
 		{
-			// Update world locked layer positions.
-			if (Layer.LayerDesc.PositionType == ELayerType::WorldLocked)
-			{
-				vr::HmdMatrix34_t HmdTransform;
-				TransformToSteamSpace(Layer.LayerDesc.Transform * InvWorldTransform, HmdTransform, WorldToMeterScale);
-				OVR_VERIFY(VROverlay->SetOverlayTransformAbsolute(Layer.OverlayHandle, VRCompositor->GetTrackingSpace(), &HmdTransform));
-			}
-
 			// Update layer textures
 			if (Layer.bUpdateTexture || (Layer.LayerDesc.Flags & LAYER_FLAG_TEX_CONTINUOUS_UPDATE))
 			{
