Reset VROverlay pointer when shutting down.

diff --git a/Engine/Plugins/Runtime/Steam/SteamVR/Source/SteamVR/Private/SteamVRHMD.cpp b/Engine/Plugins/Runtime/Steam/SteamVR/Source/SteamVR/Private/SteamVRHMD.cpp
--- a/Engine/Plugins/Runtime/Steam/SteamVR/Source/SteamVR/Private/SteamVRHMD.cpp
+++ b/Engine/Plugins/Runtime/Steam/SteamVR/Source/SteamVR/Private/SteamVRHMD.cpp
@@ -1725,6 +1725,7 @@ void FSteamVRHMD::Shutdown()
 		VRSystem = nullptr;
 		VRCompositor = nullptr;
 		VRChaperone = nullptr;
+		VROverlay = nullptr;
 
 		SteamVRPlugin->Reset();
 	}
diff --git a/Engine/Plugins/Runtime/Steam/SteamVR/Source/SteamVR/Private/SteamVRStereoLayers.cpp b/Engine/Plugins/Runtime/Steam/SteamVR/Source/SteamVR/Private/SteamVRStereoLayers.cpp
--- a/Engine/Plugins/Runtime/Steam/SteamVR/Source/SteamVR/Private/SteamVRStereoLayers.cpp
+++ b/Engine/Plugins/Runtime/Steam/SteamVR/Source/SteamVR/Private/SteamVRStereoLayers.cpp
@@ -369,7 +369,6 @@ IStereoLayers* FSteamVRHMD::GetStereoLayers()
 		return FHeadMountedDisplayBase::GetStereoLayers();
 	}
 
-	ensure(VROverlay);
 	if (VROverlay)
 	{
 		return this;
