UE-73022 Crash entering VR Preview while using a Stereo Layer - Assertion failed: XRCamera.IsValid() in SteamVRStereoLayers.cpp
-Null checking XRCamera rather than asserting it is not null.

diff --git a/Engine/Plugins/Runtime/Steam/SteamVR/Source/SteamVR/Private/SteamVRStereoLayers.cpp b/Engine/Plugins/Runtime/Steam/SteamVR/Source/SteamVR/Private/SteamVRStereoLayers.cpp
--- a/Engine/Plugins/Runtime/Steam/SteamVR/Source/SteamVR/Private/SteamVRStereoLayers.cpp
+++ b/Engine/Plugins/Runtime/Steam/SteamVR/Source/SteamVR/Private/SteamVRStereoLayers.cpp
@@ -184,9 +184,8 @@ void FSteamVRHMD::UpdateLayer(struct FSteamVRLayer& Layer, uint32 LayerId, bool
 	FQuat AdjustedPlayerOrientation = BaseOrientation.Inverse() * PlayerOrientation;
 	AdjustedPlayerOrientation.Normalize();
 
-	check(XRCamera.IsValid());
 	FVector AdjustedPlayerLocation = PlayerLocation;
-	if (XRCamera->GetUseImplicitHMDPosition())
+	if (XRCamera && XRCamera->GetUseImplicitHMDPosition())
 	{
 		AdjustedPlayerLocation -= BaseOrientation.Inverse().RotateVector(GameTrackingFrame.DevicePosition[IXRTrackingSystem::HMDDeviceId]);
 	}
