World locked layer positions need to be updated every frame in case the player game world position has changed.
Partial revert of CL#5921736.

diff --git a/Engine/Plugins/Runtime/Steam/SteamVR/Source/SteamVR/Private/SteamVRStereoLayers.cpp b/Engine/Plugins/Runtime/Steam/SteamVR/Source/SteamVR/Private/SteamVRStereoLayers.cpp
--- a/Engine/Plugins/Runtime/Steam/SteamVR/Source/SteamVR/Private/SteamVRStereoLayers.cpp
+++ b/Engine/Plugins/Runtime/Steam/SteamVR/Source/SteamVR/Private/SteamVRStereoLayers.cpp
@@ -181,25 +181,12 @@ void FSteamVRHMD::UpdateLayer(struct FSteamVRLayer& Layer, uint32 LayerId, bool
 	// This will preserve the correct order between layers with negative and positive priorities
 	OVR_VERIFY(VROverlay->SetOverlaySortOrder(Layer.OverlayHandle, Layer.LayerDesc.Priority-INT32_MIN));
 
-	FQuat AdjustedPlayerOrientation = BaseOrientation.Inverse() * PlayerOrientation;
-	AdjustedPlayerOrientation.Normalize();
-
-	FVector AdjustedPlayerLocation = PlayerLocation;
-	if (XRCamera && XRCamera->GetUseImplicitHMDPosition())
-	{
-		AdjustedPlayerLocation -= BaseOrientation.Inverse().RotateVector(GameTrackingFrame.DevicePosition[IXRTrackingSystem::HMDDeviceId]);
-	}
-
-	FTransform InvWorldTransform = FTransform(AdjustedPlayerOrientation, AdjustedPlayerLocation).Inverse();
-
 	// Transform
 	switch (Layer.LayerDesc.PositionType)
 	{
 	case ELayerType::WorldLocked:
 	{
-		vr::HmdMatrix34_t HmdTransform;
-		TransformToSteamSpace(Layer.LayerDesc.Transform * InvWorldTransform, HmdTransform, WorldToMeterScale);
-		OVR_VERIFY(VROverlay->SetOverlayTransformAbsolute(Layer.OverlayHandle, VRCompositor->GetTrackingSpace(), &HmdTransform));
+		// World locked layer positions are updated every frame.
 		break;
 	}
 	case ELayerType::TrackerLocked:
@@ -219,7 +206,6 @@ void FSteamVRHMD::UpdateLayer(struct FSteamVRLayer& Layer, uint32 LayerId, bool
 	};
 }
 
-
 //=============================================================================
 void FSteamVRHMD::UpdateStereoLayers_RenderThread()
 {
@@ -238,11 +224,36 @@ void FSteamVRHMD::UpdateStereoLayers_RenderThread()
 	typedef TTuple<vr::VROverlayHandle_t, int32, bool /* bIsFaceLocked */> LayerPriorityInfo;
 	TArray<LayerPriorityInfo> LayerPriorities;
 
+	const float WorldToMeterScale = GetWorldToMetersScale();
+	FTransform InvWorldTransform{ENoInit::NoInit};
+	{
+		// Calculate a transform to translate from world to tracker relative coordinates.
+		FQuat AdjustedPlayerOrientation = BaseOrientation.Inverse() * PlayerOrientation;
+		AdjustedPlayerOrientation.Normalize();
+
+		check(XRCamera.IsValid());
+		FVector AdjustedPlayerLocation = PlayerLocation;
+		if (XRCamera->GetUseImplicitHMDPosition())
+		{
+			AdjustedPlayerLocation -= BaseOrientation.Inverse().RotateVector(RenderTrackingFrame.DevicePosition[IXRTrackingSystem::HMDDeviceId]);
+		}
+
+		InvWorldTransform = FTransform(AdjustedPlayerOrientation, AdjustedPlayerLocation).Inverse();
+	}
+
 	// We have loop through all layers every frame, in case we have world locked layers or continuously updated textures.
 	ForEachLayer([&](uint32 /* unused */, FSteamVRLayer& Layer)
 	{
 		if (Layer.OverlayHandle != vr::k_ulOverlayHandleInvalid)
 		{
+			// Update world locked layer positions.
+			if (Layer.LayerDesc.PositionType == ELayerType::WorldLocked)
+			{
+				vr::HmdMatrix34_t HmdTransform;
+				TransformToSteamSpace(Layer.LayerDesc.Transform * InvWorldTransform, HmdTransform, WorldToMeterScale);
+				OVR_VERIFY(VROverlay->SetOverlayTransformAbsolute(Layer.OverlayHandle, VRCompositor->GetTrackingSpace(), &HmdTransform));
+			}
+
 			// Update layer textures
 			if (Layer.bUpdateTexture || (Layer.LayerDesc.Flags & LAYER_FLAG_TEX_CONTINUOUS_UPDATE))
 			{
