Tentative fixes for #jira UE-57680:
* Reduce copying by using TArray::Emplace instead of Add. (Copy constructor is not actually defined, so it messes up the shared resource refcounting.)
* Invoke Reset() when destroying FSteamVRModel and FSteamVRTexture instances.
* Added logging on freeing shared resources to assist with debugging future crash reports.
Cherrypicked CL7099699 from Dev-VR

diff --git a/Engine/Plugins/Runtime/Steam/SteamVR/Source/SteamVR/Private/SteamVRAssetManager.cpp b/Engine/Plugins/Runtime/Steam/SteamVR/Source/SteamVR/Private/SteamVRAssetManager.cpp
--- a/Engine/Plugins/Runtime/Steam/SteamVR/Source/SteamVR/Private/SteamVRAssetManager.cpp
+++ b/Engine/Plugins/Runtime/Steam/SteamVR/Source/SteamVR/Private/SteamVRAssetManager.cpp
@@ -134,7 +134,7 @@ public:
 	{
 		if (IsPending())
 		{
-			FSharedSteamVRResource& SharedResource = SharedResources.FindOrAdd(ResourceId);
+			FSharedSteamVRResource& SharedResource = SharedResources[ResourceId]; // The shared resource was already created in the constructor
 			if (SharedResource.RawResource != nullptr)
 			{
 				RawResource = SharedResource.RawResource;
@@ -226,6 +226,7 @@ template<>
 void TSteamVRModel::FreeResource(vr::IVRRenderModels* VRModelManager)
 {
 #if STEAMVR_SUPPORTED_PLATFORMS
+	UE_LOG(LogSteamVR, Log, TEXT("TSteamVRModel::FreeResource: Freeing render model instance. ResourceId: %s"), *ResourceId);
 	VRModelManager->FreeRenderModel(RawResource);
 #endif
 }
@@ -246,6 +247,7 @@ template<>
 void TSteamVRTexture::FreeResource(vr::IVRRenderModels* VRModelManager)
 {
 #if STEAMVR_SUPPORTED_PLATFORMS
+	UE_LOG(LogSteamVR, Log, TEXT("TSteamVRTexture::FreeResource: Freeing raw texture resource. ResourceId: %d"), ResourceId);
 	VRModelManager->FreeTexture(RawResource);
 #endif
 }
@@ -260,6 +262,10 @@ public:
 		: TSteamVRModel(ResID, bKickOffLoad)
 	{}
 
+	virtual ~FSteamVRModel()
+	{
+		Reset();
+	}
 public:
 	bool GetRawMeshData(float UEMeterScale, FSteamVRMeshData& MeshDataOut);
 };
@@ -330,6 +336,11 @@ public:
 		: TSteamVRTexture(ResID, bKickOffLoad)
 	{}
 
+	~FSteamVRTexture()
+	{
+		Reset();
+	}
+
 	int32 GetResourceID() const { return ResourceId; }
 
 public:
@@ -438,7 +449,7 @@ int32 FSteamVRAsyncMeshLoader::EnqueMeshLoad(const FString& ModelName)
 	if (!ModelName.IsEmpty())
 	{
 		++PendingLoadCount;
-		MeshIndex = EnqueuedMeshes.Add( FSteamVRModel(ModelName) );
+		MeshIndex = EnqueuedMeshes.Emplace( ModelName );
 	}
 	return MeshIndex;
 }
@@ -458,29 +469,31 @@ void FSteamVRAsyncMeshLoader::Tick(float /*DeltaTime*/)
 	for (int32 SubMeshIndex = 0; SubMeshIndex < EnqueuedMeshes.Num(); ++SubMeshIndex)
 	{
 		FSteamVRModel& ModelResource = EnqueuedMeshes[SubMeshIndex];
-
-		vr::RenderModel_t* RenderModel = ModelResource.TickAsyncLoad();
-		if (!ModelResource.IsPending())
+		if (ModelResource.IsPending())
 		{
-			--PendingLoadCount;
-
-			if (!RenderModel)
+			vr::RenderModel_t* RenderModel = ModelResource.TickAsyncLoad();
+			if (!ModelResource.IsPending())
 			{
-				// valid index + missing RenderModel => signifies failure
-				OnLoadComplete(SubMeshIndex);
-			}
+				--PendingLoadCount;
+
+				if (!RenderModel)
+				{
+					// valid index + missing RenderModel => signifies failure
+					OnLoadComplete(SubMeshIndex);
+				}
 #if STEAMVR_SUPPORTED_PLATFORMS
-			// if we've already loaded and converted the texture
-			else if (ConstructedTextures.Contains(RenderModel->diffuseTextureId))
-			{
-				OnLoadComplete(SubMeshIndex);
-			}
+				// if we've already loaded and converted the texture
+				else if (ConstructedTextures.Contains(RenderModel->diffuseTextureId))
+				{
+					OnLoadComplete(SubMeshIndex);
+				}
 #endif // STEAMVR_SUPPORTED_PLATFORMS
-			else if (!EnqueueTextureLoad(SubMeshIndex, RenderModel))
-			{
-				// if we fail to load the texture, we'll have to do without it
-				OnLoadComplete(SubMeshIndex);
-			}			
+				else if (!EnqueueTextureLoad(SubMeshIndex, RenderModel))
+				{
+					// if we fail to load the texture, we'll have to do without it
+					OnLoadComplete(SubMeshIndex);
+				}
+			}
 		}
 	}
 
@@ -549,7 +562,7 @@ bool FSteamVRAsyncMeshLoader::EnqueueTextureLoad(int32 SubMeshIndex, vr::RenderM
 		bLoadEnqueued = true;
 
 		// load will be kicked off later on in Tick() loop (no need to do it twice in the same tick)
-		int32 TextureIndex = EnqueuedTextures.Add( FSteamVRTexture(RenderModel->diffuseTextureId, /*bKickOffLoad =*/false) );
+		int32 TextureIndex = EnqueuedTextures.Emplace( RenderModel->diffuseTextureId, /*bKickOffLoad =*/false );
 		PendingTextureLoads.Add(TextureIndex, SubMeshIndex);
 	}
 #endif
