Merging critical SteamVR fixes from Dev-VR

diff --git a/Engine/Plugins/Runtime/Steam/SteamVR/Source/SteamVR/Private/SteamVRHMD.cpp b/Engine/Plugins/Runtime/Steam/SteamVR/Source/SteamVR/Private/SteamVRHMD.cpp
--- a/Engine/Plugins/Runtime/Steam/SteamVR/Source/SteamVR/Private/SteamVRHMD.cpp
+++ b/Engine/Plugins/Runtime/Steam/SteamVR/Source/SteamVR/Private/SteamVRHMD.cpp
@@ -128,6 +128,7 @@ public:
 	}
 #else //STEAMVR_SUPPORTED_PLATFORMS
 		: VRSystem(nullptr)
+		, VRInput(nullptr)
 	{
 		LoadOpenVRModule();
 	}
@@ -151,6 +152,11 @@ public:
 		return VRSystem;
 	}
 	
+	virtual vr::IVRInput* GetVRInput() const override
+	{
+		return VRInput;
+	}
+
 	bool LoadOpenVRModule()
 	{
 #if PLATFORM_WINDOWS
@@ -229,6 +235,7 @@ public:
 	virtual void Reset() override
 	{
 		VRSystem = nullptr;
+		VRInput = nullptr;
 	}
 	
 	uint64 GetGraphicsAdapterLuid() override
@@ -349,7 +356,8 @@ public:
 
 private:
 	vr::IVRSystem* VRSystem;
-	
+	vr::IVRInput* VRInput;
+
 	void* OpenVRDLLHandle;
 #endif // STEAMVR_SUPPORTED_PLATFORMS
 };
@@ -363,6 +371,7 @@ TSharedPtr< class IXRTrackingSystem, ESPMode::ThreadSafe > FSteamVRPlugin::Creat
 	if( SteamVRHMD->IsInitialized() )
 	{
 		VRSystem = SteamVRHMD->GetVRSystem();
+		VRInput = SteamVRHMD->GetVRInput();
 		return SteamVRHMD;
 	}
 #endif//STEAMVR_SUPPORTED_PLATFORMS
@@ -1539,6 +1548,7 @@ bool FSteamVRHMD::Startup()
 		return false;
 	}
 
+	
 	// attach to the compositor	
 	//VRCompositor = (vr::IVRCompositor*)vr::VR_GetGenericInterface(vr::IVRCompositor_Version, &HmdErr);
 	int CompositorConnectRetries = 10;
@@ -1612,6 +1622,14 @@ bool FSteamVRHMD::Startup()
 		vr::EVRInitError RenderModelsErr = vr::VRInitError_None;
 		VRRenderModels = (vr::IVRRenderModels*)(*VRGetGenericInterfaceFn)(vr::IVRRenderModels_Version, &RenderModelsErr);
 
+		vr::EVRInitError InputErr = vr::VRInitError_None;
+		VRInput = (vr::IVRInput*)(*VRGetGenericInterfaceFn)(vr::IVRInput_Version, &InputErr);
+		if ((VRInput == nullptr) || (InputErr != vr::VRInitError_None))
+		{
+			UE_LOG(LogHMD, Log, TEXT("Failed to initialize VRInput.  Error: %d"), (int32)InputErr);
+		}
+
+
 #if PLATFORM_MAC
 		if (IsMetalPlatform(GMaxRHIShaderPlatform))
 		{
diff --git a/Engine/Plugins/Runtime/Steam/SteamVR/Source/SteamVR/Private/SteamVRHMD.h b/Engine/Plugins/Runtime/Steam/SteamVR/Source/SteamVR/Private/SteamVRHMD.h
--- a/Engine/Plugins/Runtime/Steam/SteamVR/Source/SteamVR/Private/SteamVRHMD.h
+++ b/Engine/Plugins/Runtime/Steam/SteamVR/Source/SteamVR/Private/SteamVRHMD.h
@@ -309,6 +309,7 @@ public:
 		VulkanBridge(FSteamVRHMD* plugin);
 
 		virtual void FinishRendering() override;
+		virtual void UpdateViewport(const FViewport& Viewport, FRHIViewport* InViewportRHI) override;
 		virtual void Reset() override;
 
 	protected:
@@ -369,6 +370,7 @@ public:
 	bool IsInitialized() const;
 
 	vr::IVRSystem* GetVRSystem() const { return VRSystem; }
+	vr::IVRInput* GetVRInput() const { return VRInput; }
 	vr::IVRRenderModels* GetRenderModelManager() const { return VRRenderModels; }
 
 protected:
@@ -569,6 +571,7 @@ private:
 	vr::IVROverlay* VROverlay;
 	vr::IVRChaperone* VRChaperone;
 	vr::IVRRenderModels* VRRenderModels;
+	vr::IVRInput* VRInput;
 
 	FString DisplayId;
 
diff --git a/Engine/Plugins/Runtime/Steam/SteamVR/Source/SteamVR/Private/SteamVRRender.cpp b/Engine/Plugins/Runtime/Steam/SteamVR/Source/SteamVR/Private/SteamVRRender.cpp
--- a/Engine/Plugins/Runtime/Steam/SteamVR/Source/SteamVR/Private/SteamVRRender.cpp
+++ b/Engine/Plugins/Runtime/Steam/SteamVR/Source/SteamVR/Private/SteamVRRender.cpp
@@ -272,6 +272,12 @@ void FSteamVRHMD::VulkanBridge::Reset()
 
 }
 
+void FSteamVRHMD::VulkanBridge::UpdateViewport(const FViewport& Viewport, FRHIViewport* InViewportRHI)
+{
+	RenderTargetTexture = Viewport.GetRenderTargetTexture();
+	check(IsValidRef(RenderTargetTexture));
+}
+
 FSteamVRHMD::OpenGLBridge::OpenGLBridge(FSteamVRHMD* plugin):
 	BridgeBaseImpl(plugin),
 	RenderTargetTexture(0)
diff --git a/Engine/Plugins/Runtime/Steam/SteamVR/Source/SteamVR/Public/ISteamVRPlugin.h b/Engine/Plugins/Runtime/Steam/SteamVR/Source/SteamVR/Public/ISteamVRPlugin.h
--- a/Engine/Plugins/Runtime/Steam/SteamVR/Source/SteamVR/Public/ISteamVRPlugin.h
+++ b/Engine/Plugins/Runtime/Steam/SteamVR/Source/SteamVR/Public/ISteamVRPlugin.h
@@ -49,11 +49,18 @@ public:
 
 #if STEAMVR_SUPPORTED_PLATFORMS
 	/**
-	 * Get the IVRSystem* that was previously set by the HMD implemenentation.
-	 *
-	 * @return The pointer if the HMD has been initialized, otherwise nullptr 
-	 */
-	virtual vr::IVRSystem* GetVRSystem() const=0;
+	* Get the IVRSystem* that was previously set by the HMD implemenentation.
+	*
+	* @return The pointer if the HMD has been initialized, otherwise nullptr
+	*/
+	virtual vr::IVRSystem* GetVRSystem() const = 0;
+
+	/**
+	* Get the IVRInput* that was previously set by the HMD implemenentation.
+	*
+	* @return The pointer if the HMD has been initialized, otherwise nullptr
+	*/
+	virtual vr::IVRInput* GetVRInput() const = 0;
 
 	/**
 	 * Resets the plugin, in case of early VR system shutdown
diff --git a/Engine/Plugins/Runtime/Steam/SteamVR/Source/SteamVRController/Private/SteamVRController.cpp b/Engine/Plugins/Runtime/Steam/SteamVR/Source/SteamVRController/Private/SteamVRController.cpp
--- a/Engine/Plugins/Runtime/Steam/SteamVR/Source/SteamVRController/Private/SteamVRController.cpp
+++ b/Engine/Plugins/Runtime/Steam/SteamVR/Source/SteamVRController/Private/SteamVRController.cpp
@@ -8,8 +8,21 @@
 #include "Engine/Engine.h"
 #include "EngineGlobals.h"
 #include "GenericPlatform/IInputInterface.h"
+#include "GameFramework/InputSettings.h"
 #include "../../SteamVR/Private/SteamVRHMD.h"
 #include "SteamVRControllerLibrary.h" // for ESteamVRTouchDPadMapping
+#include "HAL/FileManagerGeneric.h"
+
+#if WITH_EDITOR
+#include "Editor.h"
+#endif
+
+#include "UObject/UObjectIterator.h"
+#include "GameFramework/PlayerInput.h"
+
+#include "Misc/FileHelper.h"
+#include "Serialization/JsonReader.h"
+#include "Serialization/JsonSerializer.h"
 
 #define LOCTEXT_NAMESPACE "SteamVRController"
 
@@ -53,6 +66,15 @@ static TAutoConsoleVariable<int32> CVarSwapHands(
 	TEXT(" 1: swap left and right buttons"),
 	ECVF_Cheat);
 
+static TAutoConsoleVariable<int32> CVarEnableVRInput(
+	TEXT("vr.SteamVR.EnableVRInput"),
+	0,
+	TEXT("Enable the new Steam VR Input interface for mapping actions to events.\n")
+	TEXT(" 0: use the legacy input mapping (default)\n")
+	TEXT(" 1: use the new Input API. You will have to define input bindings for the controllers you want to support."),
+	ECVF_ReadOnly);
+
+
 namespace SteamVRControllerKeyNames
 {
 	const FGamepadKeyNames::Type Touch0("Steam_Touch_0");
@@ -114,7 +136,7 @@ public:
 	/** Total number of motion controllers we'll support */
 	static const int32 MaxControllers = MaxUnrealControllers * CONTROLLERS_PER_PLAYER;
 
-	/** The maximum number of Special hand designations available to use for generic trackers 
+	/** The maximum number of Special hand designations available to use for generic trackers
 	 *  Casting enums directly, so if the input model changes, this won't silently be invalid
 	 */
 	static const int32 MaxSpecialDesignations = (int32)EControllerHand::Special_9 - (int32)EControllerHand::Special_1 + 1;
@@ -144,12 +166,31 @@ public:
 
 	FSteamVRController(const TSharedRef< FGenericApplicationMessageHandler >& InMessageHandler)
 		: MessageHandler(InMessageHandler),
-		  SteamVRPlugin(nullptr)
+		SteamVRPlugin(nullptr)
 	{
 #if STEAMVRCONTROLLER_SUPPORTED_PLATFORMS
+		static const auto CVar = IConsoleManager::Get().FindTConsoleVariableDataInt(TEXT("vr.SteamVR.EnableVRInput"));
+		bEnableVRInput = (CVar->GetValueOnGameThread() != 0) ? true : false;
+
 		FMemory::Memzero(ControllerStates, sizeof(ControllerStates));
+		NumControllersMapped = 0;
+		NumTrackersMapped = 0;
+
+		InitialButtonRepeatDelay = 0.2f;
+		ButtonRepeatDelay = 0.1f;
+
+		InitControllerMappings();
+		InitLegacyControllerKeys();
+		BuildActionManifest();
 
-		for (int32 i=0; i < vr::k_unMaxTrackedDeviceCount; ++i)
+		IModularFeatures::Get().RegisterModularFeature(GetModularFeatureName(), this);
+#endif // STEAMVRINPUT_SUPPORTED_PLATFORMS
+	}
+
+#if STEAMVRCONTROLLER_SUPPORTED_PLATFORMS
+	void InitControllerMappings()
+	{
+		for (int32 i = 0; i < vr::k_unMaxTrackedDeviceCount; ++i)
 		{
 			DeviceToControllerMap[i] = INDEX_NONE;
 		}
@@ -158,7 +199,7 @@ public:
 		{
 			for (int32 HandIndex = 0; HandIndex < vr::k_unMaxTrackedDeviceCount; ++HandIndex)
 			{
-				UnrealControllerIdAndHandToDeviceIdMap[UnrealControllerIndex][HandIndex] = INDEX_NONE;
+				SetUnrealControllerIdToControllerIndex(UnrealControllerIndex, (EControllerHand) HandIndex, INDEX_NONE);
 			}
 		}
 
@@ -166,101 +207,108 @@ public:
 		{
 			Count = 0;
 		}
+	}
 
-		NumControllersMapped = 0;
-		NumTrackersMapped = 0;
 
-		InitialButtonRepeatDelay = 0.2f;
-		ButtonRepeatDelay = 0.1f;
+	void InitLegacyControllerKeys()
+	{
 
-		Buttons[ (int32)EControllerHand::Left ][ ESteamVRControllerButton::System ] = FGamepadKeyNames::SpecialLeft;
-		Buttons[ (int32)EControllerHand::Left ][ ESteamVRControllerButton::ApplicationMenu ] = FGamepadKeyNames::MotionController_Left_Shoulder;
-		Buttons[ (int32)EControllerHand::Left ][ ESteamVRControllerButton::TouchPadPress ] = FGamepadKeyNames::MotionController_Left_Thumbstick;
-		Buttons[ (int32)EControllerHand::Left ][ ESteamVRControllerButton::TouchPadTouch ] = SteamVRControllerKeyNames::Touch0;
-		Buttons[ (int32)EControllerHand::Left ][ ESteamVRControllerButton::TriggerPress ] = FGamepadKeyNames::MotionController_Left_Trigger;
-		Buttons[ (int32)EControllerHand::Left ][ ESteamVRControllerButton::Grip ] = FGamepadKeyNames::MotionController_Left_Grip1;
-
-		Buttons[ (int32)EControllerHand::Right ][ ESteamVRControllerButton::System ] = FGamepadKeyNames::SpecialRight;
-		Buttons[ (int32)EControllerHand::Right ][ ESteamVRControllerButton::ApplicationMenu ] = FGamepadKeyNames::MotionController_Right_Shoulder;
-		Buttons[ (int32)EControllerHand::Right ][ ESteamVRControllerButton::TouchPadPress ] = FGamepadKeyNames::MotionController_Right_Thumbstick;
-		Buttons[ (int32)EControllerHand::Right ][ ESteamVRControllerButton::TouchPadTouch ] = SteamVRControllerKeyNames::Touch1;
-		Buttons[ (int32)EControllerHand::Right ][ ESteamVRControllerButton::TriggerPress ] = FGamepadKeyNames::MotionController_Right_Trigger;
-		Buttons[ (int32)EControllerHand::Right ][ ESteamVRControllerButton::Grip ] = FGamepadKeyNames::MotionController_Right_Grip1;
-
-		// Init Left & Right, TouchPadUp/Down/Left/Right button mappings
-		SetTouchDPadMapping(DefaultDPadMapping);
-
-		Buttons[ (int32)EControllerHand::Pad ][ ESteamVRControllerButton::System ] = SteamVRControllerKeyNames::GenericSystem;
-		Buttons[ (int32)EControllerHand::Pad ][ ESteamVRControllerButton::ApplicationMenu ] = SteamVRControllerKeyNames::GenericMenu;
-		Buttons[ (int32)EControllerHand::Pad ][ ESteamVRControllerButton::TouchPadPress ] = SteamVRControllerKeyNames::GenericTouchpad;
-		Buttons[ (int32)EControllerHand::Pad ][ ESteamVRControllerButton::TouchPadTouch ] = FGamepadKeyNames::Invalid;
-		Buttons[ (int32)EControllerHand::Pad ][ ESteamVRControllerButton::TriggerPress ] = SteamVRControllerKeyNames::GenericTrigger;
-		Buttons[ (int32)EControllerHand::Pad ][ ESteamVRControllerButton::Grip ] = SteamVRControllerKeyNames::GenericGrip;
-		Buttons[ (int32)EControllerHand::Pad ][ ESteamVRControllerButton::TouchPadUp ] = FGamepadKeyNames::Invalid;
-		Buttons[ (int32)EControllerHand::Pad ][ ESteamVRControllerButton::TouchPadDown ] = FGamepadKeyNames::Invalid;
-		Buttons[ (int32)EControllerHand::Pad ][ ESteamVRControllerButton::TouchPadLeft ] = FGamepadKeyNames::Invalid;
-		Buttons[ (int32)EControllerHand::Pad ][ ESteamVRControllerButton::TouchPadRight ] = FGamepadKeyNames::Invalid;
-
-		Buttons[ (int32)EControllerHand::ExternalCamera ][ ESteamVRControllerButton::System ] = FGamepadKeyNames::Invalid;
-		Buttons[ (int32)EControllerHand::ExternalCamera ][ ESteamVRControllerButton::ApplicationMenu ] = FGamepadKeyNames::Invalid;
-		Buttons[ (int32)EControllerHand::ExternalCamera ][ ESteamVRControllerButton::TouchPadPress ] = FGamepadKeyNames::Invalid;
-		Buttons[ (int32)EControllerHand::ExternalCamera ][ ESteamVRControllerButton::TouchPadTouch ] = FGamepadKeyNames::Invalid;
-		Buttons[ (int32)EControllerHand::ExternalCamera ][ ESteamVRControllerButton::TriggerPress ] = FGamepadKeyNames::Invalid;
-		Buttons[ (int32)EControllerHand::ExternalCamera ][ ESteamVRControllerButton::Grip ] = FGamepadKeyNames::Invalid;
-		Buttons[ (int32)EControllerHand::ExternalCamera ][ ESteamVRControllerButton::TouchPadUp ] = FGamepadKeyNames::Invalid;
-		Buttons[ (int32)EControllerHand::ExternalCamera ][ ESteamVRControllerButton::TouchPadDown ] = FGamepadKeyNames::Invalid;
-		Buttons[ (int32)EControllerHand::ExternalCamera ][ ESteamVRControllerButton::TouchPadLeft ] = FGamepadKeyNames::Invalid;
-		Buttons[ (int32)EControllerHand::ExternalCamera ][ ESteamVRControllerButton::TouchPadRight ] = FGamepadKeyNames::Invalid;
-
-		Buttons[(int32)EControllerHand::Gun][ESteamVRControllerButton::System] = SteamVRControllerKeyNames::GenericSystem;
-		Buttons[(int32)EControllerHand::Gun][ESteamVRControllerButton::ApplicationMenu] = SteamVRControllerKeyNames::GenericMenu;
-		Buttons[(int32)EControllerHand::Gun][ESteamVRControllerButton::TouchPadPress] = FGamepadKeyNames::Invalid;
-		Buttons[(int32)EControllerHand::Gun][ESteamVRControllerButton::TouchPadTouch] = FGamepadKeyNames::Invalid;
-		Buttons[(int32)EControllerHand::Gun][ESteamVRControllerButton::TriggerPress] = SteamVRControllerKeyNames::GenericTrigger;
-		Buttons[(int32)EControllerHand::Gun][ESteamVRControllerButton::Grip] = SteamVRControllerKeyNames::GenericGrip;
-		Buttons[(int32)EControllerHand::Gun][ESteamVRControllerButton::TouchPadUp] = FGamepadKeyNames::Invalid;
-		Buttons[(int32)EControllerHand::Gun][ESteamVRControllerButton::TouchPadDown] = FGamepadKeyNames::Invalid;
-		Buttons[(int32)EControllerHand::Gun][ESteamVRControllerButton::TouchPadLeft] = FGamepadKeyNames::Invalid;
-		Buttons[(int32)EControllerHand::Gun][ESteamVRControllerButton::TouchPadRight] = FGamepadKeyNames::Invalid;
-
-		for (int32 SpecialIndex = (int32)EControllerHand::Special_1; SpecialIndex <= (int32)EControllerHand::Special_9; ++SpecialIndex)
-		{
-			Buttons[SpecialIndex][ESteamVRControllerButton::System] = SteamVRControllerKeyNames::GenericSystem;
-			Buttons[SpecialIndex][ESteamVRControllerButton::ApplicationMenu] = SteamVRControllerKeyNames::GenericMenu;
-			Buttons[SpecialIndex][ESteamVRControllerButton::TouchPadPress] = SteamVRControllerKeyNames::GenericTouchpad;
-			Buttons[SpecialIndex][ESteamVRControllerButton::TouchPadTouch] = FGamepadKeyNames::Invalid;
-			Buttons[SpecialIndex][ESteamVRControllerButton::TriggerPress] = SteamVRControllerKeyNames::GenericTrigger;
-			Buttons[SpecialIndex][ESteamVRControllerButton::Grip] = SteamVRControllerKeyNames::GenericGrip;
-			Buttons[SpecialIndex][ESteamVRControllerButton::TouchPadUp] = FGamepadKeyNames::Invalid;
-			Buttons[SpecialIndex][ESteamVRControllerButton::TouchPadDown] = FGamepadKeyNames::Invalid;
-			Buttons[SpecialIndex][ESteamVRControllerButton::TouchPadLeft] = FGamepadKeyNames::Invalid;
-			Buttons[SpecialIndex][ESteamVRControllerButton::TouchPadRight] = FGamepadKeyNames::Invalid;
-		}
-
-		EKeys::AddKey(FKeyDetails(SteamVRControllerKeys::SteamVR_Knuckles_Left_HandGrip, LOCTEXT("SteamVR_Knuckles_Left_HandGrip", "SteamVR Knuckles (L) Hand Grip CapSense"), FKeyDetails::GamepadKey | FKeyDetails::FloatAxis));
-		EKeys::AddKey(FKeyDetails(SteamVRControllerKeys::SteamVR_Knuckles_Left_IndexGrip, LOCTEXT("SteamVR_Knuckles_Left_IndexGrip", "SteamVR Knuckles (L) Index Grip CapSense"), FKeyDetails::GamepadKey | FKeyDetails::FloatAxis));
-		EKeys::AddKey(FKeyDetails(SteamVRControllerKeys::SteamVR_Knuckles_Left_MiddleGrip, LOCTEXT("SteamVR_Knuckles_Left_MiddleGrip", "SteamVR Knuckles (L) Middle Grip CapSense"), FKeyDetails::GamepadKey | FKeyDetails::FloatAxis));
-		EKeys::AddKey(FKeyDetails(SteamVRControllerKeys::SteamVR_Knuckles_Left_RingGrip, LOCTEXT("SteamVR_Knuckles_Left_RingGrip", "SteamVR Knuckles (L) Ring Grip CapSense"), FKeyDetails::GamepadKey | FKeyDetails::FloatAxis));
-		EKeys::AddKey(FKeyDetails(SteamVRControllerKeys::SteamVR_Knuckles_Left_PinkyGrip, LOCTEXT("SteamVR_Knuckles_Left_PinkyGrip", "SteamVR Knuckles (L) Pinky Grip CapSense"), FKeyDetails::GamepadKey | FKeyDetails::FloatAxis));
-
-		EKeys::AddKey(FKeyDetails(SteamVRControllerKeys::SteamVR_Knuckles_Right_HandGrip, LOCTEXT("SteamVR_Knuckles_Right_HandGrip", "SteamVR Knuckles (R) Hand Grip CapSense"), FKeyDetails::GamepadKey | FKeyDetails::FloatAxis));
-		EKeys::AddKey(FKeyDetails(SteamVRControllerKeys::SteamVR_Knuckles_Right_IndexGrip, LOCTEXT("SteamVR_Knuckles_Right_IndexGrip", "SteamVR Knuckles (R) Index Grip CapSense"), FKeyDetails::GamepadKey | FKeyDetails::FloatAxis));
-		EKeys::AddKey(FKeyDetails(SteamVRControllerKeys::SteamVR_Knuckles_Right_MiddleGrip, LOCTEXT("SteamVR_Knuckles_Right_MiddleGrip", "SteamVR Knuckles (R) Middle Grip CapSense"), FKeyDetails::GamepadKey | FKeyDetails::FloatAxis));
-		EKeys::AddKey(FKeyDetails(SteamVRControllerKeys::SteamVR_Knuckles_Right_RingGrip, LOCTEXT("SteamVR_Knuckles_Right_RingGrip", "SteamVR Knuckles (R) Ring Grip CapSense"), FKeyDetails::GamepadKey | FKeyDetails::FloatAxis));
-		EKeys::AddKey(FKeyDetails(SteamVRControllerKeys::SteamVR_Knuckles_Right_PinkyGrip, LOCTEXT("SteamVR_Knuckles_Right_PinkyGrip", "SteamVR Knuckles (R) Pinky Grip CapSense"), FKeyDetails::GamepadKey | FKeyDetails::FloatAxis));
+		if (!bEnableVRInput)
+		{
+			Buttons[(int32)EControllerHand::Left][ESteamVRControllerButton::System] = FGamepadKeyNames::SpecialLeft;
+			Buttons[(int32)EControllerHand::Left][ESteamVRControllerButton::ApplicationMenu] = FGamepadKeyNames::MotionController_Left_Shoulder;
+			Buttons[(int32)EControllerHand::Left][ESteamVRControllerButton::TouchPadPress] = FGamepadKeyNames::MotionController_Left_Thumbstick;
+			Buttons[(int32)EControllerHand::Left][ESteamVRControllerButton::TouchPadTouch] = SteamVRControllerKeyNames::Touch0;
+			Buttons[(int32)EControllerHand::Left][ESteamVRControllerButton::TriggerPress] = FGamepadKeyNames::MotionController_Left_Trigger;
+			Buttons[(int32)EControllerHand::Left][ESteamVRControllerButton::Grip] = FGamepadKeyNames::MotionController_Left_Grip1;
+
+			Buttons[(int32)EControllerHand::Right][ESteamVRControllerButton::System] = FGamepadKeyNames::SpecialRight;
+			Buttons[(int32)EControllerHand::Right][ESteamVRControllerButton::ApplicationMenu] = FGamepadKeyNames::MotionController_Right_Shoulder;
+			Buttons[(int32)EControllerHand::Right][ESteamVRControllerButton::TouchPadPress] = FGamepadKeyNames::MotionController_Right_Thumbstick;
+			Buttons[(int32)EControllerHand::Right][ESteamVRControllerButton::TouchPadTouch] = SteamVRControllerKeyNames::Touch1;
+			Buttons[(int32)EControllerHand::Right][ESteamVRControllerButton::TriggerPress] = FGamepadKeyNames::MotionController_Right_Trigger;
+			Buttons[(int32)EControllerHand::Right][ESteamVRControllerButton::Grip] = FGamepadKeyNames::MotionController_Right_Grip1;
+
+			// Init Left & Right, TouchPadUp/Down/Left/Right button mappings
+			SetTouchDPadMapping(DefaultDPadMapping);
+
+			Buttons[(int32)EControllerHand::Pad][ESteamVRControllerButton::System] = SteamVRControllerKeyNames::GenericSystem;
+			Buttons[(int32)EControllerHand::Pad][ESteamVRControllerButton::ApplicationMenu] = SteamVRControllerKeyNames::GenericMenu;
+			Buttons[(int32)EControllerHand::Pad][ESteamVRControllerButton::TouchPadPress] = SteamVRControllerKeyNames::GenericTouchpad;
+			Buttons[(int32)EControllerHand::Pad][ESteamVRControllerButton::TouchPadTouch] = FGamepadKeyNames::Invalid;
+			Buttons[(int32)EControllerHand::Pad][ESteamVRControllerButton::TriggerPress] = SteamVRControllerKeyNames::GenericTrigger;
+			Buttons[(int32)EControllerHand::Pad][ESteamVRControllerButton::Grip] = SteamVRControllerKeyNames::GenericGrip;
+			Buttons[(int32)EControllerHand::Pad][ESteamVRControllerButton::TouchPadUp] = FGamepadKeyNames::Invalid;
+			Buttons[(int32)EControllerHand::Pad][ESteamVRControllerButton::TouchPadDown] = FGamepadKeyNames::Invalid;
+			Buttons[(int32)EControllerHand::Pad][ESteamVRControllerButton::TouchPadLeft] = FGamepadKeyNames::Invalid;
+			Buttons[(int32)EControllerHand::Pad][ESteamVRControllerButton::TouchPadRight] = FGamepadKeyNames::Invalid;
+
+			Buttons[(int32)EControllerHand::ExternalCamera][ESteamVRControllerButton::System] = FGamepadKeyNames::Invalid;
+			Buttons[(int32)EControllerHand::ExternalCamera][ESteamVRControllerButton::ApplicationMenu] = FGamepadKeyNames::Invalid;
+			Buttons[(int32)EControllerHand::ExternalCamera][ESteamVRControllerButton::TouchPadPress] = FGamepadKeyNames::Invalid;
+			Buttons[(int32)EControllerHand::ExternalCamera][ESteamVRControllerButton::TouchPadTouch] = FGamepadKeyNames::Invalid;
+			Buttons[(int32)EControllerHand::ExternalCamera][ESteamVRControllerButton::TriggerPress] = FGamepadKeyNames::Invalid;
+			Buttons[(int32)EControllerHand::ExternalCamera][ESteamVRControllerButton::Grip] = FGamepadKeyNames::Invalid;
+			Buttons[(int32)EControllerHand::ExternalCamera][ESteamVRControllerButton::TouchPadUp] = FGamepadKeyNames::Invalid;
+			Buttons[(int32)EControllerHand::ExternalCamera][ESteamVRControllerButton::TouchPadDown] = FGamepadKeyNames::Invalid;
+			Buttons[(int32)EControllerHand::ExternalCamera][ESteamVRControllerButton::TouchPadLeft] = FGamepadKeyNames::Invalid;
+			Buttons[(int32)EControllerHand::ExternalCamera][ESteamVRControllerButton::TouchPadRight] = FGamepadKeyNames::Invalid;
+
+			Buttons[(int32)EControllerHand::Gun][ESteamVRControllerButton::System] = SteamVRControllerKeyNames::GenericSystem;
+			Buttons[(int32)EControllerHand::Gun][ESteamVRControllerButton::ApplicationMenu] = SteamVRControllerKeyNames::GenericMenu;
+			Buttons[(int32)EControllerHand::Gun][ESteamVRControllerButton::TouchPadPress] = FGamepadKeyNames::Invalid;
+			Buttons[(int32)EControllerHand::Gun][ESteamVRControllerButton::TouchPadTouch] = FGamepadKeyNames::Invalid;
+			Buttons[(int32)EControllerHand::Gun][ESteamVRControllerButton::TriggerPress] = SteamVRControllerKeyNames::GenericTrigger;
+			Buttons[(int32)EControllerHand::Gun][ESteamVRControllerButton::Grip] = SteamVRControllerKeyNames::GenericGrip;
+			Buttons[(int32)EControllerHand::Gun][ESteamVRControllerButton::TouchPadUp] = FGamepadKeyNames::Invalid;
+			Buttons[(int32)EControllerHand::Gun][ESteamVRControllerButton::TouchPadDown] = FGamepadKeyNames::Invalid;
+			Buttons[(int32)EControllerHand::Gun][ESteamVRControllerButton::TouchPadLeft] = FGamepadKeyNames::Invalid;
+			Buttons[(int32)EControllerHand::Gun][ESteamVRControllerButton::TouchPadRight] = FGamepadKeyNames::Invalid;
+
+			for (int32 SpecialIndex = (int32)EControllerHand::Special_1; SpecialIndex <= (int32)EControllerHand::Special_9; ++SpecialIndex)
+			{
+				Buttons[SpecialIndex][ESteamVRControllerButton::System] = SteamVRControllerKeyNames::GenericSystem;
+				Buttons[SpecialIndex][ESteamVRControllerButton::ApplicationMenu] = SteamVRControllerKeyNames::GenericMenu;
+				Buttons[SpecialIndex][ESteamVRControllerButton::TouchPadPress] = SteamVRControllerKeyNames::GenericTouchpad;
+				Buttons[SpecialIndex][ESteamVRControllerButton::TouchPadTouch] = FGamepadKeyNames::Invalid;
+				Buttons[SpecialIndex][ESteamVRControllerButton::TriggerPress] = SteamVRControllerKeyNames::GenericTrigger;
+				Buttons[SpecialIndex][ESteamVRControllerButton::Grip] = SteamVRControllerKeyNames::GenericGrip;
+				Buttons[SpecialIndex][ESteamVRControllerButton::TouchPadUp] = FGamepadKeyNames::Invalid;
+				Buttons[SpecialIndex][ESteamVRControllerButton::TouchPadDown] = FGamepadKeyNames::Invalid;
+				Buttons[SpecialIndex][ESteamVRControllerButton::TouchPadLeft] = FGamepadKeyNames::Invalid;
+				Buttons[SpecialIndex][ESteamVRControllerButton::TouchPadRight] = FGamepadKeyNames::Invalid;
+			}
 
-		IModularFeatures::Get().RegisterModularFeature(GetModularFeatureName(), this);
-#endif // STEAMVRCONTROLLER_SUPPORTED_PLATFORMS
+			EKeys::AddKey(FKeyDetails(SteamVRControllerKeys::SteamVR_Knuckles_Left_HandGrip, LOCTEXT("SteamVR_Knuckles_Left_HandGrip", "SteamVR Knuckles (L) Hand Grip CapSense"), FKeyDetails::GamepadKey | FKeyDetails::FloatAxis));
+			EKeys::AddKey(FKeyDetails(SteamVRControllerKeys::SteamVR_Knuckles_Left_IndexGrip, LOCTEXT("SteamVR_Knuckles_Left_IndexGrip", "SteamVR Knuckles (L) Index Grip CapSense"), FKeyDetails::GamepadKey | FKeyDetails::FloatAxis));
+			EKeys::AddKey(FKeyDetails(SteamVRControllerKeys::SteamVR_Knuckles_Left_MiddleGrip, LOCTEXT("SteamVR_Knuckles_Left_MiddleGrip", "SteamVR Knuckles (L) Middle Grip CapSense"), FKeyDetails::GamepadKey | FKeyDetails::FloatAxis));
+			EKeys::AddKey(FKeyDetails(SteamVRControllerKeys::SteamVR_Knuckles_Left_RingGrip, LOCTEXT("SteamVR_Knuckles_Left_RingGrip", "SteamVR Knuckles (L) Ring Grip CapSense"), FKeyDetails::GamepadKey | FKeyDetails::FloatAxis));
+			EKeys::AddKey(FKeyDetails(SteamVRControllerKeys::SteamVR_Knuckles_Left_PinkyGrip, LOCTEXT("SteamVR_Knuckles_Left_PinkyGrip", "SteamVR Knuckles (L) Pinky Grip CapSense"), FKeyDetails::GamepadKey | FKeyDetails::FloatAxis));
+
+			EKeys::AddKey(FKeyDetails(SteamVRControllerKeys::SteamVR_Knuckles_Right_HandGrip, LOCTEXT("SteamVR_Knuckles_Right_HandGrip", "SteamVR Knuckles (R) Hand Grip CapSense"), FKeyDetails::GamepadKey | FKeyDetails::FloatAxis));
+			EKeys::AddKey(FKeyDetails(SteamVRControllerKeys::SteamVR_Knuckles_Right_IndexGrip, LOCTEXT("SteamVR_Knuckles_Right_IndexGrip", "SteamVR Knuckles (R) Index Grip CapSense"), FKeyDetails::GamepadKey | FKeyDetails::FloatAxis));
+			EKeys::AddKey(FKeyDetails(SteamVRControllerKeys::SteamVR_Knuckles_Right_MiddleGrip, LOCTEXT("SteamVR_Knuckles_Right_MiddleGrip", "SteamVR Knuckles (R) Middle Grip CapSense"), FKeyDetails::GamepadKey | FKeyDetails::FloatAxis));
+			EKeys::AddKey(FKeyDetails(SteamVRControllerKeys::SteamVR_Knuckles_Right_RingGrip, LOCTEXT("SteamVR_Knuckles_Right_RingGrip", "SteamVR Knuckles (R) Ring Grip CapSense"), FKeyDetails::GamepadKey | FKeyDetails::FloatAxis));
+			EKeys::AddKey(FKeyDetails(SteamVRControllerKeys::SteamVR_Knuckles_Right_PinkyGrip, LOCTEXT("SteamVR_Knuckles_Right_PinkyGrip", "SteamVR Knuckles (R) Pinky Grip CapSense"), FKeyDetails::GamepadKey | FKeyDetails::FloatAxis));
+		}
 	}
+#endif // STEAMVRCONTROLLER_SUPPORTED_PLATFORMS
 
 	virtual ~FSteamVRController()
 	{
 #if STEAMVRCONTROLLER_SUPPORTED_PLATFORMS
+#if WITH_EDITOR
+		if (ActionMappingsChangedHandle.IsValid())
+		{
+			FEditorDelegates::OnActionAxisMappingsChanged.Remove(ActionMappingsChangedHandle);
+			ActionMappingsChangedHandle.Reset();
+		}
+#endif
 		IModularFeatures::Get().UnregisterModularFeature(GetModularFeatureName(), this);
 #endif // STEAMVRCONTROLLER_SUPPORTED_PLATFORMS
 	}
 
-	virtual void Tick( float DeltaTime ) override
+	virtual void Tick(float DeltaTime) override
 	{
 #if STEAMVRCONTROLLER_SUPPORTED_PLATFORMS
 		vr::IVRSystem* VRSystem = GetVRSystem();
@@ -270,13 +318,26 @@ public:
 			RegisterDeviceChanges(VRSystem);
 			DetectHandednessSwap(VRSystem);
 		}
-
 #endif // STEAMVRCONTROLLER_SUPPORTED_PLATFORMS
 	}
 
 	virtual void SendControllerEvents() override
 	{
 #if STEAMVRCONTROLLER_SUPPORTED_PLATFORMS
+		if (bEnableVRInput)
+		{
+			SendActionInputEvents();
+		}
+		else
+		{
+			SendLegacyControllerEvents();
+		}
+#endif
+	}
+
+#if STEAMVRCONTROLLER_SUPPORTED_PLATFORMS
+	void SendLegacyControllerEvents()
+	{
 		vr::VRControllerState_t VRControllerState;
 
 		vr::IVRSystem* VRSystem = GetVRSystem();
@@ -458,8 +519,102 @@ public:
 				}
 			}
 		}
+	}
 #endif // STEAMVRCONTROLLER_SUPPORTED_PLATFORMS
+
+#if STEAMVRCONTROLLER_SUPPORTED_PLATFORMS
+	void SendActionInputEvents()
+	{
+		vr::IVRInput* VRInput = GetVRInput();
+
+		if (VRInput != nullptr)
+		{
+			vr::VRActiveActionSet_t ActiveActionSets[] = {
+				{
+					MainActionSet,
+					vr::k_ulInvalidInputValueHandle,
+					vr::k_ulInvalidActionSetHandle
+				}
+			};
+			vr::EVRInputError Err = VRInput->UpdateActionState(ActiveActionSets, sizeof(vr::VRActiveActionSet_t), 1);
+			if (Err != vr::VRInputError_None)
+			{
+				UE_LOG(LogSteamVRController, Warning, TEXT("UpdateActionState returned error: %d"), (int32)Err);
+				return;
+			}
+
+			for (auto& Action : Actions)
+			{
+				switch (Action.Type)
+				{
+				case FSteamVRAction::Boolean:
+				{
+					vr::InputDigitalActionData_t Data;
+					Err = VRInput->GetDigitalActionData(Action.Handle, &Data, sizeof(Data), vr::k_ulInvalidInputValueHandle);
+					if (Err == vr::VRInputError_None)
+					{
+						if (Data.bState != Action.bState)
+						{
+							Action.bState = Data.bState;
+							if (Action.bState)
+							{
+								MessageHandler->OnControllerButtonPressed(Action.ActionKey_X, 0, /*IsRepeat =*/false);
+							}
+							else
+							{
+								MessageHandler->OnControllerButtonReleased(Action.ActionKey_X, 0, /*IsRepeat =*/false);
+							}
+						}
+					}
+					// If the current error is the same as the last frame's error, don't log it again to avoid spamming the log
+					else if (Err != Action.LastError)
+					{
+						UE_LOG(LogSteamVRController, Warning, TEXT("GetDigitalActionData for %s returned error: %d"), *Action.Name.ToString(), (int32)Err);
+
+					}
+					Action.LastError = Err;
+				}
+				break;
+				case FSteamVRAction::Vector1:
+				case FSteamVRAction::Vector2:
+				case FSteamVRAction::Vector3:
+				{
+					vr::InputAnalogActionData_t Data;
+					Err = VRInput->GetAnalogActionData(Action.Handle, &Data, sizeof(Data), vr::k_ulInvalidInputValueHandle);
+					if (Err == vr::VRInputError_None)
+					{
+						if (!Action.ActionKey_X.IsNone() && Data.x != Action.Value.X)
+						{
+							Action.Value.X = Data.x;
+							MessageHandler->OnControllerAnalog(Action.ActionKey_X, 0, Action.Value.X);
+						}
+						if (!Action.ActionKey_Y.IsNone() && Data.y != Action.Value.Y)
+						{
+							Action.Value.Y = Data.y;
+							MessageHandler->OnControllerAnalog(Action.ActionKey_Y, 0, Action.Value.Y);
+						}
+						if (!Action.ActionKey_Z.IsNone() && Data.z != Action.Value.Z)
+						{
+							Action.Value.Z = Data.z;
+							MessageHandler->OnControllerAnalog(Action.ActionKey_Z, 0, Action.Value.Z);
+						}
+					}
+					// If the current error is the same as the last frame's error, don't log it again to avoid spamming the log
+					else if (Err != Action.LastError)
+					{
+						UE_LOG(LogSteamVRController, Warning, TEXT("GetAnalogActionData for %s returned error: %d"), *Action.Name.ToString(), (int32)Err);
+					}
+					Action.LastError = Err;
+				}
+				break;
+				default:
+					// Other action types are currently unsupported.
+					break;
+				}
+			}
+		}
 	}
+#endif // STEAMVRCONTROLLER_SUPPORTED_PLATFORMS
 
 	void SetTouchDPadMapping(ESteamVRTouchDPadMapping NewMapping)
 	{
@@ -510,10 +665,16 @@ public:
 	}
 
 #if STEAMVRCONTROLLER_SUPPORTED_PLATFORMS
-	int32 UnrealControllerIdToControllerIndex( const int32 UnrealControllerId, const EControllerHand Hand ) const
+	int32 UnrealControllerIdToControllerIndex(const int32 UnrealControllerId, const EControllerHand Hand) const
 	{
 		return UnrealControllerIdAndHandToDeviceIdMap[UnrealControllerId][(int32)Hand];
 	}
+
+	void SetUnrealControllerIdToControllerIndex(const int32 UnrealControllerId, const EControllerHand Hand, int32 value)
+	{
+		UnrealControllerIdAndHandToDeviceIdMap[UnrealControllerId][(int32)Hand] = value;
+	}
+
 #endif // STEAMVRCONTROLLER_SUPPORTED_PLATFORMS
 	
 	void SetChannelValue(int32 UnrealControllerId, FForceFeedbackChannelType ChannelType, float Value) override
@@ -641,7 +802,7 @@ public:
  		FSteamVRHMD* SteamVRHMD = GetSteamVRHMD();
  		if (SteamVRHMD)
  		{
-			int32 DeviceId = UnrealControllerIdAndHandToDeviceIdMap[ControllerIndex][(int32)DeviceHand];
+			int32 DeviceId = UnrealControllerIdToControllerIndex(ControllerIndex, DeviceHand);
  			FQuat DeviceOrientation = FQuat::Identity;
 			// Steam handles WorldToMetersScale when it reads the controller posrot, so we do not need to use it again here.  Debugging found that they are the same.
 			RetVal = SteamVRHMD->GetCurrentPose(DeviceId, DeviceOrientation, OutPosition);
@@ -660,7 +821,7 @@ public:
 		FSteamVRHMD* SteamVRHMD = GetSteamVRHMD();
  		if (SteamVRHMD)
  		{
-			int32 DeviceId = UnrealControllerIdAndHandToDeviceIdMap[ControllerIndex][(int32)DeviceHand];
+			int32 DeviceId = UnrealControllerIdToControllerIndex(ControllerIndex, DeviceHand);
 			TrackingStatus = SteamVRHMD->GetControllerTrackingStatus(DeviceId);
  		}
 #endif // STEAMVRCONTROLLER_SUPPORTED_PLATFORMS
@@ -691,7 +852,17 @@ public:
 	static ESteamVRTouchDPadMapping DefaultDPadMapping;
 private:
 
-	inline vr::IVRSystem* GetVRSystem()
+	inline vr::IVRInput* GetVRInput() const
+	{
+		if (SteamVRPlugin == nullptr)
+		{
+			SteamVRPlugin = &FModuleManager::LoadModuleChecked<ISteamVRPlugin>(TEXT("SteamVR"));
+		}
+
+		return SteamVRPlugin->GetVRInput();
+	}
+
+	inline vr::IVRSystem* GetVRSystem() const
 	{
 		if (SteamVRPlugin == nullptr)
 		{
@@ -800,7 +971,7 @@ private:
 
 		++NumControllersMapped;
 
-		UnrealControllerIdAndHandToDeviceIdMap[DeviceToControllerMap[DeviceIndex]][(int32)ControllerStates[DeviceIndex].Hand] = DeviceIndex;
+		SetUnrealControllerIdToControllerIndex(DeviceToControllerMap[DeviceIndex], ControllerStates[DeviceIndex].Hand, DeviceIndex);
 
 		return true;
 	}
@@ -822,8 +993,8 @@ private:
 
 				int32 ControllerIndex = DeviceToControllerMap[LeftDeviceIndex];
 
-				UnrealControllerIdAndHandToDeviceIdMap[ControllerIndex][(int32)EControllerHand::Left] = LeftDeviceIndex;
-				UnrealControllerIdAndHandToDeviceIdMap[ControllerIndex][(int32)EControllerHand::Right] = RightDeviceIndex;
+				SetUnrealControllerIdToControllerIndex(ControllerIndex, EControllerHand::Left, LeftDeviceIndex);
+				SetUnrealControllerIdToControllerIndex(ControllerIndex, EControllerHand::Right, RightDeviceIndex);
 			}
 		}
 	}
@@ -887,7 +1058,7 @@ private:
 		++NumTrackersMapped;
 		UE_LOG(LogSteamVRController, Log, TEXT("Tracker device %i is being assigned unreal hand: Special %i, for player %i"), DeviceIndex, NumTrackersMapped, GENERIC_TRACKER_PLAYER_NUM);
 
-		UnrealControllerIdAndHandToDeviceIdMap[DeviceToControllerMap[DeviceIndex]][(int32)ControllerStates[DeviceIndex].Hand] = DeviceIndex;
+		SetUnrealControllerIdToControllerIndex(DeviceToControllerMap[DeviceIndex], ControllerStates[DeviceIndex].Hand, DeviceIndex);
 
 		return true;
 	}
@@ -908,13 +1079,404 @@ private:
 	void UnregisterDevice(uint32 DeviceIndex)
 	{
 		// undo the mappings
-		UnrealControllerIdAndHandToDeviceIdMap[DeviceToControllerMap[DeviceIndex]][(int32)ControllerStates[DeviceIndex].Hand] = INDEX_NONE;
+		SetUnrealControllerIdToControllerIndex(DeviceToControllerMap[DeviceIndex], ControllerStates[DeviceIndex].Hand, INDEX_NONE);
 		DeviceToControllerMap[DeviceIndex] = INDEX_NONE;
 
 		// re-zero out the controller state
 		FMemory::Memzero(&ControllerStates[DeviceIndex], sizeof(FControllerState));
 	}
 
+	// Hack to prefer emitting MotionController keys for action events
+	static bool MatchKeyNamePrefix(const FKey& Key, const TCHAR* Prefix)
+	{
+		return Key.GetFName().ToString().StartsWith(Prefix);
+	};
+
+	static bool MatchKeyNameSuffix(const FKey& Key, const TCHAR* Suffix)
+	{
+		return Key.GetFName().ToString().EndsWith(Suffix);
+	};
+
+	// Finds an axis key mapping from a list of mapping with the following preferences:
+	// 1. Tries to find a FloatAxis key that starts with "MotionController" and ends with "X"
+	// 2. Tries to find a FloatAxis key that starts with "MotionController" and ends with "Y"
+	// 3. Any to find a FloatAxis key that starts with "MotionController"
+	// 4. Any FloatAxis that ends with "X"
+	// 5. Any FloatAxis that ends with "Y"
+	// 6. Any FloatAxis
+	// 7. Any Valid key.
+	// If case 1 or 3 is matched bOutIsXAxis will be set to true. 
+	static FName FindAxisKeyMapping(TArray<FInputAxisKeyMapping>& Mappings, bool& bOutIsXAxis)
+	{
+		FInputAxisKeyMapping* Found = nullptr;
+		bOutIsXAxis = false;
+		// First filter out all floatAxes, as all except the 5th case require a float axis.
+		TArray<FInputAxisKeyMapping> FloatMappings = Mappings.FilterByPredicate([](const FInputAxisKeyMapping& Mapping)
+		{
+			return Mapping.Key.IsFloatAxis();
+		});
+
+		// If there were no float axis key bindings, return the first valid mapping
+		if (FloatMappings.Num() == 0)
+		{
+			Found = Mappings.FindByPredicate([](FInputAxisKeyMapping& Mapping)
+			{
+				return Mapping.Key.IsValid();
+			});
+
+			if (Found != nullptr)
+			{
+				return Found->Key.GetFName();
+			}
+			else
+			{
+				return FName();
+			}
+		}
+
+		// Then get all mappings with keys starting with "MotionController"
+		TArray<FInputAxisKeyMapping> MotionControllerMappings = FloatMappings.FilterByPredicate([](const FInputAxisKeyMapping& Mapping)
+		{
+			return MatchKeyNamePrefix(Mapping.Key, TEXT("MotionController"));
+		});
+
+		// If there are no MotionController keys, search through all FloatAxes:
+		TArray<FInputAxisKeyMapping>& MappingsSubset = MotionControllerMappings.Num() == 0 ? FloatMappings : MotionControllerMappings;
+
+		Found = FloatMappings.FindByPredicate([](const FInputAxisKeyMapping& Mapping)
+		{
+			return MatchKeyNameSuffix(Mapping.Key, TEXT("X"));
+		});
+		if (Found != nullptr)
+		{
+			bOutIsXAxis = true;
+			return Found->Key.GetFName();
+		}
+
+		Found = FloatMappings.FindByPredicate([](const FInputAxisKeyMapping& Mapping)
+		{
+			return MatchKeyNameSuffix(Mapping.Key, TEXT("Y"));
+		});
+		if (Found != nullptr)
+		{
+			return Found->Key.GetFName();
+		}
+
+
+		Found = FloatMappings.FindByPredicate([](const FInputAxisKeyMapping& Mapping)
+		{
+			return Mapping.Key.IsValid();
+		});
+		if (Found != nullptr)
+		{
+			return Found->Key.GetFName();
+		}
+		else
+		{
+			return FName();
+		}
+	}
+
+	/** Returns the concatenation of two strings, skipping all characters at the beginning of string B that match the beginning of string A and
+	    all characters at the end of string A that match the end of string B.
+		Example: passing in "MoveUpAction" and "MoveRightAction" should result in "MoveUpRightAction"
+		If the strings have no common suffix or prefix, the result will simply be the concatenation of both strings.
+		If the strings are identical, returns the first string.
+
+		The algorithm treats separator characters ' ', '_' and '/' differently. If either the suffix begins with one or the prefix ends with one,
+		the function will keep one of them in the resulting string.
+		Example "move_up_action" and "move_right_action" will result in "move_up_right_action" and not "move_upright_action"
+	   */
+	static FString MergeActionNames(const FString& A, const FString& B)
+	{
+		if (A.Equals(B, ESearchCase::CaseSensitive))
+		{
+			return A;
+		}
+		const int LastA = A.Len() - 1;
+		const int LastB = B.Len() - 1;
+		const int MinLen = (LastA < LastB) ? A.Len() : B.Len();
+
+		int CommonPrefix = 0;
+		int CommonSuffix = 0;
+		for (;CommonPrefix < MinLen && A[CommonPrefix] == B[CommonPrefix]; CommonPrefix++)
+		{
+			/* intentionally blank */
+		}
+
+		for (;CommonSuffix < MinLen && A[LastA - CommonSuffix] == B[LastB - CommonSuffix]; CommonSuffix++)
+		{
+			/* intentionally blank */
+		}
+
+		// If either the common prefix ends with or the common suffix begins with a space, an underscore or a dash, keep one of them.
+		if (CommonPrefix > 0 && (A[CommonPrefix - 1] == TEXT(' ') || A[CommonPrefix - 1] == TEXT('_') || A[CommonPrefix - 1] == TEXT('-')))
+		{
+			CommonPrefix--;
+		}
+		else if (CommonSuffix > 0 && (A[LastA - CommonSuffix + 1] == TEXT(' ') || A[LastA - CommonSuffix + 1] == TEXT('_') || A[LastA - CommonSuffix + 1] == TEXT('-')))
+		{
+			CommonSuffix--;
+		}
+
+		return A.LeftChop(CommonSuffix) + B.RightChop(CommonPrefix);
+	}
+
+	void BuildActionManifest()
+	{
+		vr::IVRInput* VRInput;
+
+		if (bEnableVRInput && (VRInput = GetVRInput()) != nullptr)
+		{
+			Actions.Empty();
+			auto InputSettings = GetDefault<UInputSettings>();
+			if (InputSettings != nullptr)
+			{
+				TArray<FName> ActionNames;
+				InputSettings->GetActionNames(ActionNames);
+				for (const auto& ActionName : ActionNames)
+				{
+					TArray<FInputActionKeyMapping> Mappings;
+					InputSettings->GetActionMappingByName(ActionName, Mappings);
+
+					FInputActionKeyMapping* KeyMapping = Mappings.FindByPredicate([](FInputActionKeyMapping& Mapping) 
+					{ 
+						return MatchKeyNamePrefix(Mapping.Key, TEXT("MotionController")); 
+					});
+
+					if (KeyMapping == nullptr)
+					{
+						KeyMapping = Mappings.FindByPredicate([](FInputActionKeyMapping& Mapping) { return Mapping.Key.IsValid(); });
+					}
+
+					if (KeyMapping != nullptr)
+					{
+						FString ActionPath = FString("/actions/main/in") / ActionName.ToString();
+						Actions.Add(FSteamVRAction( ActionPath, ActionName, KeyMapping->Key.GetFName(), false ));
+					}
+				}
+
+				TArray<FName> AxisNames;
+				InputSettings->GetAxisNames(AxisNames);
+				for (const auto& AxisName : AxisNames)
+				{
+					bool bIsXAxis = false;
+					TArray<FInputAxisKeyMapping> Mappings;
+					InputSettings->GetAxisMappingByName(AxisName, Mappings);
+
+					FName KeyName = FindAxisKeyMapping(Mappings, bIsXAxis);
+
+					if (!KeyName.IsNone())
+					{
+						FString ActionPath = FString("/actions/main/in") / AxisName.ToString() + TEXT("_axis");
+						Actions.Add(FSteamVRAction( ActionPath, AxisName, KeyName, 0.0f ));
+
+						// If the current axis is bound to an X axis, find the corresponding Y axis binding and create
+						// a combined vector2 action from them (and if there were Z axes, create vector3 actions.)
+						if (bIsXAxis)
+						{
+							FName YKeyName = FName(*(KeyName.ToString().LeftChop(1) + TEXT('Y')));
+							FName ZKeyName = FName(*(KeyName.ToString().LeftChop(1) + TEXT('Z')));
+							FName YAxisName, ZAxisName;
+
+							for (const auto& InnerAxisName : AxisNames)
+							{
+								TArray<FInputAxisKeyMapping> InnerMappings;
+								InputSettings->GetAxisMappingByName(InnerAxisName, InnerMappings);
+
+								if (YAxisName.IsNone() && 
+									InnerMappings.ContainsByPredicate([&YKeyName](auto& Mapping)
+									{
+										return Mapping.Key.GetFName() == YKeyName;
+									})
+								)
+								{
+									YAxisName = InnerAxisName;
+								}
+
+								if (ZAxisName.IsNone() && 
+									InnerMappings.ContainsByPredicate([&ZKeyName](auto& Mapping)
+									{
+										return Mapping.Key.GetFName() == ZKeyName;
+									})
+								)
+								{
+									ZAxisName = InnerAxisName;
+								}
+
+								if (!YAxisName.IsNone() && !ZAxisName.IsNone())
+								{
+									break;
+								}
+
+							}
+
+							if (!YAxisName.IsNone())
+							{
+								FString CombinedAxisName = MergeActionNames(AxisName.ToString(), YAxisName.ToString());
+								if (!ZAxisName.IsNone())
+								{
+									CombinedAxisName = MergeActionNames(CombinedAxisName, ZAxisName.ToString());
+									FString CombinedActionPath = FString("/actions/main/in") / CombinedAxisName + TEXT("_axis3d");
+									Actions.Add(FSteamVRAction(CombinedActionPath, FName(*CombinedAxisName), KeyName, YKeyName, ZKeyName, FVector::ZeroVector));
+								}
+								else
+								{
+									FString CombinedActionPath = FString("/actions/main/in") / CombinedAxisName + TEXT("_axis2d");
+									Actions.Add(FSteamVRAction(CombinedActionPath, FName(*CombinedAxisName), KeyName, YKeyName, FVector2D::ZeroVector));
+								}
+							}
+									
+						}
+					}
+				}
+
+				// Open console
+				{
+					const FKey* ConsoleKey = InputSettings->ConsoleKeys.FindByPredicate([](FKey& Key) { return Key.IsValid(); });
+					if (ConsoleKey != nullptr)
+					{
+						Actions.Add(FSteamVRAction( FString("/actions/main/in/open_console"), FName(TEXT("Open Console")), ConsoleKey->GetFName(), false ));
+					}
+				}
+			}
+
+			if (Actions.Num() > 0)
+			{
+				// The steamvr_actions.json file is generated from internal data, so it goes to GeneratedConfig directory.
+				const FString ManifestPath = FPaths::GeneratedConfigDir() / TEXT("steamvr_actions.json");
+				// The default bindings files need to be generated by the developer, so they are stored in the project config directory.
+				const FString BindingsDir = FPaths::ProjectConfigDir() / TEXT("SteamVRBindings");
+
+				TSharedPtr<FJsonObject> DescriptionsObject = MakeShareable(new FJsonObject);
+
+				TArray<TSharedPtr<FJsonValue>> ActionsArray;
+				for (auto Action : Actions)
+				{
+					TSharedRef<FJsonObject> ActionObject = MakeShareable(new FJsonObject());
+					ActionObject->SetStringField(TEXT("name"), Action.Path);
+					ActionObject->SetStringField(TEXT("type"), Action.TypeAsString());
+					ActionObject->SetStringField(TEXT("requirement"), TEXT("optional"));
+					ActionsArray.Add(MakeShareable(new FJsonValueObject(ActionObject)));
+
+					DescriptionsObject->SetStringField(Action.Path, Action.Name.ToString());
+				}
+
+				TArray<TSharedPtr<FJsonValue>> DefaultBindings;
+				{
+					IFileManager& FileManager = FFileManagerGeneric::Get();
+
+					// Find any default bindings stored in the project bindings dir.
+					// They must be saved as <PROJECT_CONFIG_DIR>/SteamVRBindings/<CONTROLLER_TYPE>.json in order to be included in the manifest.
+					TArray<FString> FoundFiles;
+					FileManager.FindFiles(FoundFiles, *BindingsDir, TEXT("*.json"));
+					UE_LOG(LogSteamVRController, Log, TEXT("Searching for bindings files in %s"), *BindingsDir);
+					for (FString& File : FoundFiles)
+					{
+						FString ControllerType = File;
+						ControllerType.RemoveFromEnd(".json");
+						TSharedRef<FJsonObject> BindingObject = MakeShareable(new FJsonObject());
+						BindingObject->SetStringField(TEXT("controller_type"), *ControllerType);
+						BindingObject->SetStringField(TEXT("binding_url"), *FileManager.ConvertToAbsolutePathForExternalAppForRead(*(BindingsDir / File)));
+						DefaultBindings.Add(MakeShareable(new FJsonValueObject(BindingObject)));
+					}
+
+					if (DefaultBindings.Num() == 0)
+					{
+#if WITH_EDITOR
+						UE_LOG(LogSteamVRController, Log, TEXT("No default Steam VR Input bindings found. Adding an inital entry for vive_controller."));
+						
+						// Create the directory if it doesn't exist.
+						if (!FileManager.DirectoryExists(*BindingsDir))
+						{
+							FileManager.MakeDirectory(*BindingsDir);
+						}
+
+						// Telling SteamVR about the path will allow the bindings editor to automatically export the file to the correct location when
+						// pressing "Replace Default Binding".
+						FString DefaultViveController = FileManager.ConvertToAbsolutePathForExternalAppForRead(*(BindingsDir / TEXT("vive_controller.json")));
+						TSharedRef<FJsonObject> BindingObject = MakeShareable(new FJsonObject());
+						BindingObject->SetStringField(TEXT("controller_type"), TEXT("vive_controller"));
+						BindingObject->SetStringField(TEXT("binding_url"), *DefaultViveController);
+						DefaultBindings.Add(MakeShareable(new FJsonValueObject(BindingObject)));
+#else
+						UE_LOG(LogSteamVRController, Error, TEXT("No default Steam VR Input bindings found in %s."), *BindingsDir);
+#endif
+					}
+				}
+
+				TArray<TSharedPtr<FJsonValue>> ActionSets;
+				{
+					TSharedRef<FJsonObject> ActionSetObject = MakeShareable(new FJsonObject());
+					ActionSetObject->SetStringField(TEXT("name"), TEXT("/actions/main"));
+					ActionSetObject->SetStringField(TEXT("usage"), TEXT("leftright"));
+					ActionSets.Add(MakeShareable(new FJsonValueObject(ActionSetObject)));
+
+					DescriptionsObject->SetStringField(TEXT("/actions/main"), TEXT("Main Game Actions"));
+				}
+
+				DescriptionsObject->SetStringField(TEXT("language_tag"), TEXT("en"));
+				TArray<TSharedPtr<FJsonValue>> Localization;
+				{
+					Localization.Add(MakeShareable(new FJsonValueObject(DescriptionsObject)));
+				}
+
+				TSharedRef<FJsonObject> RootObject = MakeShareable(new FJsonObject());
+				RootObject->SetArrayField(TEXT("default_bindings"), DefaultBindings);
+				RootObject->SetArrayField(TEXT("actions"), ActionsArray);
+				RootObject->SetArrayField(TEXT("action_sets"), ActionSets);
+				RootObject->SetArrayField(TEXT("localization"), Localization);
+
+				// Print the JSON data to a string
+				FString OutputJsonString;
+				TSharedRef<TJsonWriter<>> JsonWriter = TJsonWriterFactory<>::Create(&OutputJsonString);
+				FJsonSerializer::Serialize(RootObject, JsonWriter);
+
+				// Save the JSON string (force UTF-8 for JSON files.)
+				if (!FFileHelper::SaveStringToFile(OutputJsonString, *ManifestPath, FFileHelper::EEncodingOptions::ForceUTF8WithoutBOM))
+				{
+					UE_LOG(LogSteamVRController, Error, TEXT("Failed to save action manifest '%s'."), *ManifestPath);
+					return;
+				}
+
+				vr::EVRInputError Err = VRInput->SetActionManifestPath(TCHAR_TO_ANSI(*IFileManager::Get().ConvertToAbsolutePathForExternalAppForRead(*ManifestPath)));
+
+				if (Err != vr::VRInputError_None)
+				{
+					UE_LOG(LogSteamVRController, Error, TEXT("Failed to pass action manifest, %s, to SteamVR. Error: %d"), *ManifestPath, (int32)Err);
+				}
+
+				// Get the action set handle for our main action set
+				Err = VRInput->GetActionSetHandle("/actions/main", &MainActionSet);
+				if (Err != vr::VRInputError_None)
+				{
+					UE_LOG(LogSteamVRController, Error, TEXT("Couldn't get main action set handle. Error: %d"), (int32)Err);
+				}
+
+				// Fill in Action handles for each registered action
+				for (auto& Action : Actions)
+				{
+					vr::VRActionHandle_t Handle;
+					Err = VRInput->GetActionHandle(TCHAR_TO_ANSI(*Action.Path), &Handle);
+					Action.Handle = Handle;
+					if (Err != vr::VRInputError_None || !Action.Handle)
+					{
+						UE_LOG(LogSteamVRController, Warning, TEXT("Couldn't get main action handle for %s. Error: %d"), *Action.Path, (int32)Err);
+					}
+				}
+#if WITH_EDITOR
+				if (!ActionMappingsChangedHandle.IsValid())
+				{
+					ActionMappingsChangedHandle = FEditorDelegates::OnActionAxisMappingsChanged.AddLambda([this]()
+					{
+						UE_LOG(LogSteamVRController, Warning, TEXT("You will need to quit and restart both SteamVR and the Editor in order to use the modified input actions or axes."));
+					});
+				}
+#endif
+			}
+		}
+	}
+
 	struct FControllerState
 	{
 		/** Which hand this controller is representing */
@@ -948,6 +1510,102 @@ private:
 		float ForceFeedbackValue;
 	};
 
+	struct FSteamVRAction
+	{
+		enum EActionType
+		{
+			Boolean,
+			Vector1,
+			Vector2,
+			Vector3,
+			Vibration,
+			Pose,
+			Skeleton,
+			Invalid
+		};
+
+		FString		Path;
+		EActionType	Type;
+		FName		Name;
+		FName		ActionKey_X;
+		FName		ActionKey_Y;
+		FName		ActionKey_Z;
+		union {
+			bool	bState;
+			FVector Value;
+		};
+
+		vr::VRActionHandle_t Handle;
+		vr::EVRInputError LastError;
+
+		FString TypeAsString()
+		{
+			static FString TypeStrings[] = {
+				TEXT("boolean"),
+				TEXT("vector1"),
+				TEXT("vector2"),
+				TEXT("vector3"),
+				TEXT("vibration"),
+				TEXT("pose"),
+				TEXT("skeleton"),
+				TEXT("")
+			};
+
+			return TypeStrings[(int)Type];
+		}
+
+		FSteamVRAction(const FString& inPath, const FName& inName, const FName& inActionKey, bool inState)
+			: Path(inPath)
+			, Type(Boolean)
+			, Name(inName)
+			, ActionKey_X(inActionKey)
+			, ActionKey_Y()
+			, ActionKey_Z()
+			, Value()
+			, Handle()
+			, LastError(vr::VRInputError_None)
+		{
+			bState = inState;
+		}
+
+		FSteamVRAction(const FString& inPath, const FName& inName, const FName& inActionKey, float inValue1D)
+			: Path(inPath)
+			, Type(Vector1)
+			, Name(inName)
+			, ActionKey_X(inActionKey)
+			, ActionKey_Y()
+			, ActionKey_Z()
+			, Value(inValue1D, 0, 0)
+			, Handle()
+			, LastError(vr::VRInputError_None)
+		{}
+
+		FSteamVRAction(const FString& inPath, const FName& inName, const FName& inActionKey_X, const FName& inActionKey_Y, const FVector2D& inValue2D)
+			: Path(inPath)
+			, Type(Vector2)
+			, Name(inName)
+			, ActionKey_X(inActionKey_X)
+			, ActionKey_Y(inActionKey_Y)
+			, ActionKey_Z()
+			, Value(inValue2D.X, inValue2D.Y, 0)
+			, Handle()
+			, LastError(vr::VRInputError_None)
+		{}
+
+		FSteamVRAction(const FString& inPath, const FName& inName, const FName& inActionKey_X, const FName& inActionKey_Y, const FName& inActionKey_Z, const FVector& inValue3D)
+			: Path(inPath)
+			, Type(Vector2)
+			, Name(inName)
+			, ActionKey_X(inActionKey_X)
+			, ActionKey_Y(inActionKey_Y)
+			, ActionKey_Z(inActionKey_Z)
+			, Value(inValue3D)
+			, Handle()
+			, LastError(vr::VRInputError_None)
+		{}
+
+	};
+
 	/** Mappings between tracked devices and 0 indexed controllers */
 	int32 NumControllersMapped;
 	int32 NumTrackersMapped;
@@ -958,6 +1616,9 @@ private:
 	/** Controller states */
 	FControllerState ControllerStates[ MaxControllers ];
 
+	TArray<FSteamVRAction> Actions;
+	vr::VRActionSetHandle_t MainActionSet;
+
 	/** Delay before sending a repeat message after a button was first pressed */
 	float InitialButtonRepeatDelay;
 
@@ -969,13 +1630,21 @@ private:
 
 	/** weak pointer to the IVRSystem owned by the HMD module */
 	TWeakPtr<vr::IVRSystem> HMDVRSystem;
+
+#if WITH_EDITOR
+	FDelegateHandle ActionMappingsChangedHandle;
+#endif
+
 #endif // STEAMVRCONTROLLER_SUPPORTED_PLATFORMS
 
+	/** Whether the VRInput API is enabled or not */
+	bool bEnableVRInput;
+
 	/** handler to send all messages to */
 	TSharedRef<FGenericApplicationMessageHandler> MessageHandler;
 
 	/** the SteamVR plugin module */
-	ISteamVRPlugin* SteamVRPlugin;
+	mutable ISteamVRPlugin* SteamVRPlugin;
 };
 
 FName FSteamVRController::DeviceTypeName(TEXT("SteamVRController"));
diff --git a/Engine/Plugins/Runtime/Steam/SteamVR/Source/SteamVRController/SteamVRController.Build.cs b/Engine/Plugins/Runtime/Steam/SteamVR/Source/SteamVRController/SteamVRController.Build.cs
--- a/Engine/Plugins/Runtime/Steam/SteamVR/Source/SteamVRController/SteamVRController.Build.cs
+++ b/Engine/Plugins/Runtime/Steam/SteamVR/Source/SteamVRController/SteamVRController.Build.cs
@@ -21,15 +21,21 @@ public class SteamVRController : ModuleRules
 			"InputDevice",
             "InputCore",
 			"HeadMountedDisplay",
-            "SteamVR"
-		});
+            "SteamVR",
+            "Json"
+        });
+
+        // 		DynamicallyLoadedModuleNames.AddRange(new string[]
+        // 		{
+        // 			"SteamVR",
+        // 		});
 
-// 		DynamicallyLoadedModuleNames.AddRange(new string[]
-// 		{
-// 			"SteamVR",
-// 		});
+        if (Target.bBuildEditor == true)
+        {
+            PrivateDependencyModuleNames.Add("UnrealEd");
+        }
 
-		AddEngineThirdPartyPrivateStaticDependencies(Target, "OpenVR");
+        AddEngineThirdPartyPrivateStaticDependencies(Target, "OpenVR");
 
         if ( Target.Platform == UnrealTargetPlatform.Win64 || Target.Platform == UnrealTargetPlatform.Win32 || (Target.Platform == UnrealTargetPlatform.Linux && Target.Architecture.StartsWith("x86_64")) )
         {
