Merging using Dev-VR_to_Release-4.21

4374589 UE-57680 [CrashReport] UE4Editor_SteamVR!FSteamVRModel::GetRawMeshData() [steamvrassetmanager.cpp:307]
-Added some more logging to try to figure out why this is crashing.  Logs the resourceID of the steamVRMesh which might tell us if this is happening with some specific hardware.  Also logs index and vertex counts and pointers to see if the data is somewhat malformed or just completely invalid.

4380355 UE-57680 [CrashReport] UE4Editor_SteamVR!FSteamVRModel::GetRawMeshData() [steamvrassetmanager.cpp:307]
-Two more log lines for steamvr shutdown.  Maybe we are trying to get the models after shutdown and get bad data?

diff --git a/Engine/Plugins/Runtime/Steam/SteamVR/Source/SteamVR/Private/SteamVRAssetManager.cpp b/Engine/Plugins/Runtime/Steam/SteamVR/Source/SteamVR/Private/SteamVRAssetManager.cpp
--- a/Engine/Plugins/Runtime/Steam/SteamVR/Source/SteamVR/Private/SteamVRAssetManager.cpp
+++ b/Engine/Plugins/Runtime/Steam/SteamVR/Source/SteamVR/Private/SteamVRAssetManager.cpp
@@ -286,6 +286,9 @@ bool FSteamVRModel::GetRawMeshData(float UEMeterScale, FSteamVRMeshData& MeshDat
 #if STEAMVR_SUPPORTED_PLATFORMS
 	if (RawResource)
 	{
+		// Logging to try to get info about a crash where the RawMeshData seems to be bad.
+		UE_LOG(LogSteamVR, Log, TEXT("FSteamVRModel::GetRawMeshData ResourceId %s rVertexData 0x%lx unVertexCount %i rIndexData 0x%lx unTriangleCount %i diffuseTextureId %i"), *ResourceId, RawResource->rVertexData, RawResource->unVertexCount, RawResource->rIndexData, RawResource->unTriangleCount, RawResource->diffuseTextureId);
+
 		const uint32 VertCount = RawResource->unVertexCount;
 		MeshDataOut.VertPositions.Empty(VertCount);
 		MeshDataOut.UVs.Empty(VertCount);
diff --git a/Engine/Plugins/Runtime/Steam/SteamVR/Source/SteamVR/Private/SteamVRHMD.cpp b/Engine/Plugins/Runtime/Steam/SteamVR/Source/SteamVR/Private/SteamVRHMD.cpp
--- a/Engine/Plugins/Runtime/Steam/SteamVR/Source/SteamVR/Private/SteamVRHMD.cpp
+++ b/Engine/Plugins/Runtime/Steam/SteamVR/Source/SteamVR/Private/SteamVRHMD.cpp
@@ -222,6 +222,7 @@ public:
 	{
 		if (OpenVRDLLHandle != nullptr)
 		{
+			UE_LOG(LogHMD, Log, TEXT("Freeing OpenVRDLLHandle."));
 			FPlatformProcess::FreeDllHandle(OpenVRDLLHandle);
 			OpenVRDLLHandle = nullptr;
 		}
@@ -1706,6 +1707,8 @@ void FSteamVRHMD::Shutdown()
 {
 	if (VRSystem != nullptr)
 	{
+		UE_LOG(LogHMD, Log, TEXT("SteamVR Shutting down."));
+
 		SplashTicker->UnregisterForMapLoad();
 		SplashTicker = nullptr;
 
