Merging using Dev-VR_to_Release-4.21

Increase max controllers to support more.

diff --git a/Engine/Plugins/Runtime/Steam/SteamVR/Source/SteamVRController/Private/SteamVRController.cpp b/Engine/Plugins/Runtime/Steam/SteamVR/Source/SteamVRController/Private/SteamVRController.cpp
--- a/Engine/Plugins/Runtime/Steam/SteamVR/Source/SteamVRController/Private/SteamVRController.cpp
+++ b/Engine/Plugins/Runtime/Steam/SteamVR/Source/SteamVRController/Private/SteamVRController.cpp
@@ -39,8 +39,6 @@ DEFINE_LOG_CATEGORY_STATIC(LogSteamVRController, Log, All);
 /** Total number of controllers in a set */
 #define CONTROLLERS_PER_PLAYER	2
 
-#define MAX_TRACKED_DEVICES (int32)EControllerHand::Special_9 - (int32)EControllerHand::Left + 1
-
 /** Player that generic trackers will be assigned to */
 #define GENERIC_TRACKER_PLAYER_NUM 0
 
@@ -134,7 +132,8 @@ public:
 	static const int32 MaxUnrealControllers = MAX_STEAMVR_CONTROLLER_PAIRS;
 
 	/** Total number of motion controllers we'll support */
-	static const int32 MaxControllers = MaxUnrealControllers * CONTROLLERS_PER_PLAYER;
+	// NOTE: This used to be MaxUnrealControllers * CONTROLLERS_PER_PLAYER, but we needed to support many more trackers than that
+	static const int32 MaxControllers = vr::k_unMaxTrackedDeviceCount;
 
 	/** The maximum number of Special hand designations available to use for generic trackers
 	 *  Casting enums directly, so if the input model changes, this won't silently be invalid
