Linux: fix ARM server build (UE-64775).

diff --git a/Engine/Plugins/Runtime/Steam/SteamVR/Source/SteamVRController/Private/SteamVRController.cpp b/Engine/Plugins/Runtime/Steam/SteamVR/Source/SteamVRController/Private/SteamVRController.cpp
--- a/Engine/Plugins/Runtime/Steam/SteamVR/Source/SteamVRController/Private/SteamVRController.cpp
+++ b/Engine/Plugins/Runtime/Steam/SteamVR/Source/SteamVRController/Private/SteamVRController.cpp
@@ -133,7 +133,12 @@ public:
 
 	/** Total number of motion controllers we'll support */
 	// NOTE: This used to be MaxUnrealControllers * CONTROLLERS_PER_PLAYER, but we needed to support many more trackers than that
-	static const int32 MaxControllers = vr::k_unMaxTrackedDeviceCount;
+	static const int32 MaxControllers =
+#if STEAMVRCONTROLLER_SUPPORTED_PLATFORMS
+		vr::k_unMaxTrackedDeviceCount;
+#else
+		0;
+#endif
 
 	/** The maximum number of Special hand designations available to use for generic trackers
 	 *  Casting enums directly, so if the input model changes, this won't silently be invalid
