Cherry picking CL#4436130 from Dev-VR.
Automatically create blank bindings for common controller types for SteamVR Input if there are no bindings.

diff --git a/Engine/Plugins/Runtime/Steam/SteamVR/Source/SteamVRController/Private/SteamVRController.cpp b/Engine/Plugins/Runtime/Steam/SteamVR/Source/SteamVRController/Private/SteamVRController.cpp
--- a/Engine/Plugins/Runtime/Steam/SteamVR/Source/SteamVRController/Private/SteamVRController.cpp
+++ b/Engine/Plugins/Runtime/Steam/SteamVR/Source/SteamVRController/Private/SteamVRController.cpp
@@ -1388,7 +1388,7 @@ private:
 					if (DefaultBindings.Num() == 0)
 					{
 #if WITH_EDITOR
-						UE_LOG(LogSteamVRController, Log, TEXT("No default Steam VR Input bindings found. Adding an inital entry for vive_controller."));
+						UE_LOG(LogSteamVRController, Log, TEXT("No default Steam VR Input bindings found. Adding inital blank entries for common controller types."));
 						
 						// Create the directory if it doesn't exist.
 						if (!FileManager.DirectoryExists(*BindingsDir))
@@ -1396,13 +1396,27 @@ private:
 							FileManager.MakeDirectory(*BindingsDir);
 						}
 
-						// Telling SteamVR about the path will allow the bindings editor to automatically export the file to the correct location when
-						// pressing "Replace Default Binding".
-						FString DefaultViveController = FileManager.ConvertToAbsolutePathForExternalAppForRead(*(BindingsDir / TEXT("vive_controller.json")));
-						TSharedRef<FJsonObject> BindingObject = MakeShareable(new FJsonObject());
-						BindingObject->SetStringField(TEXT("controller_type"), TEXT("vive_controller"));
-						BindingObject->SetStringField(TEXT("binding_url"), *DefaultViveController);
-						DefaultBindings.Add(MakeShareable(new FJsonValueObject(BindingObject)));
+						static const TCHAR* CommonControllerTypes[] = 
+						{
+							TEXT("vive"),
+							TEXT("vive_controller"),
+							TEXT("oculus_touch"),
+							TEXT("holographic_controller"),
+							TEXT("gamepad")
+						};
+
+						for (const TCHAR* Type : CommonControllerTypes)
+						{
+							// Telling SteamVR about the path will allow the bindings editor to automatically export the file to the correct location when
+							// pressing "Replace Default Binding".
+							FString BindingsPath = FileManager.ConvertToAbsolutePathForExternalAppForRead(*(BindingsDir / Type + TEXT(".json")));
+							// Creating an empty JSON file will allow editing it in place in the SteamVR Bindings Editor
+							FFileHelper::SaveStringToFile(TEXT(""), *BindingsPath, FFileHelper::EEncodingOptions::ForceUTF8WithoutBOM);
+							TSharedRef<FJsonObject> BindingObject = MakeShareable(new FJsonObject());
+							BindingObject->SetStringField(TEXT("controller_type"), Type);
+							BindingObject->SetStringField(TEXT("binding_url"), *BindingsPath);
+							DefaultBindings.Add(MakeShareable(new FJsonValueObject(BindingObject)));
+						}
 #else
 						UE_LOG(LogSteamVRController, Error, TEXT("No default Steam VR Input bindings found in %s."), *BindingsDir);
 #endif
