Defer building the SteamVR hidden area meshes until Stereo is enabled

diff --git a/Engine/Plugins/Runtime/Steam/SteamVR/Source/SteamVR/Private/SteamVRHMD.cpp b/Engine/Plugins/Runtime/Steam/SteamVR/Source/SteamVR/Private/SteamVRHMD.cpp
--- a/Engine/Plugins/Runtime/Steam/SteamVR/Source/SteamVR/Private/SteamVRHMD.cpp
+++ b/Engine/Plugins/Runtime/Steam/SteamVR/Source/SteamVR/Private/SteamVRHMD.cpp
@@ -1273,6 +1273,11 @@ bool FSteamVRHMD::EnableStereo(bool bStereo)
 	}
 
 	bStereoDesired = (IsHMDEnabled()) ? bStereo : false;
+	if (bStereoDesired && !bOcclusionMeshesBuilt)
+	{
+		SetupOcclusionMeshes();
+		bOcclusionMeshesBuilt = true;
+	}
 
 	// Set the viewport to match that of the HMD display
 	FSceneViewport* SceneVP = FindSceneViewport();
@@ -1493,6 +1498,7 @@ FSteamVRHMD::FSteamVRHMD(ISteamVRPlugin* InSteamVRPlugin) :
 	HmdWornState(EHMDWornState::Unknown),
 	bStereoDesired(false),
 	bStereoEnabled(false),
+	bOcclusionMeshesBuilt(false),
 	WindowMirrorBoundsWidth(2160),
 	WindowMirrorBoundsHeight(1200),
 	PixelDensity(1.0f),
@@ -1510,10 +1516,7 @@ FSteamVRHMD::FSteamVRHMD(ISteamVRPlugin* InSteamVRPlugin) :
 	VROverlay(nullptr),
 	VRChaperone(nullptr)
 {
-	if (Startup())
-	{
-		SetupOcclusionMeshes();
-	}
+	Startup();
 }
 
 FSteamVRHMD::~FSteamVRHMD()
@@ -1771,7 +1774,7 @@ static void SetupHiddenAreaMeshes(vr::IVRSystem* const VRSystem, FHMDViewMesh Re
 
 
 void FSteamVRHMD::SetupOcclusionMeshes()
-{	
+{
 	FSteamVRHMD* const Self = this;
 	ENQUEUE_RENDER_COMMAND(SetupOcclusionMeshesCmd)([Self](FRHICommandListImmediate& RHICmdList)
 	{
diff --git a/Engine/Plugins/Runtime/Steam/SteamVR/Source/SteamVR/Private/SteamVRHMD.h b/Engine/Plugins/Runtime/Steam/SteamVR/Source/SteamVR/Private/SteamVRHMD.h
--- a/Engine/Plugins/Runtime/Steam/SteamVR/Source/SteamVR/Private/SteamVRHMD.h
+++ b/Engine/Plugins/Runtime/Steam/SteamVR/Source/SteamVR/Private/SteamVRHMD.h
@@ -443,6 +443,7 @@ private:
 	void SetupOcclusionMeshes();
 
 	bool bHmdEnabled;
+	bool bOcclusionMeshesBuilt;
 	EHMDWornState::Type HmdWornState;
 	bool bStereoDesired;
 	bool bStereoEnabled;
