Cherry picking CL#4503210 from Dev-VR.
The SteamVR Controller Bindings Editor no longer accepts or allows editing empty controller binding JSON files.
The editor will now create default files for common controller types and populate them with the minimal amount of attributes required by SteamVR Input to recognize them.

diff --git a/Engine/Plugins/Runtime/Steam/SteamVR/Source/SteamVRController/Private/SteamVRController.cpp b/Engine/Plugins/Runtime/Steam/SteamVR/Source/SteamVRController/Private/SteamVRController.cpp
--- a/Engine/Plugins/Runtime/Steam/SteamVR/Source/SteamVRController/Private/SteamVRController.cpp
+++ b/Engine/Plugins/Runtime/Steam/SteamVR/Source/SteamVRController/Private/SteamVRController.cpp
@@ -1226,6 +1226,80 @@ private:
 		return A.LeftChop(CommonSuffix) + B.RightChop(CommonPrefix);
 	}
 
+#if WITH_EDITOR
+	void BuildDefaultActionBindings(const FString& BindingsDir, TArray<TSharedPtr<FJsonValue>>& InOutDefaultBindings)
+	{
+		IFileManager& FileManager = FFileManagerGeneric::Get();
+
+		TSet<FString> ExistingBindings;
+		for (const TSharedPtr<FJsonValue>& Value : InOutDefaultBindings)
+		{
+			const TSharedPtr<FJsonObject>* Object;
+			FString ControllerType;
+			if (Value.IsValid() && Value->TryGetObject(Object) && (*Object)->TryGetStringField(TEXT("controller_type"), ControllerType) && !ControllerType.IsEmpty())
+			{
+				ExistingBindings.Emplace(ControllerType);
+			}
+		}
+
+		// Create the directory if it doesn't exist.
+		if (!FileManager.DirectoryExists(*BindingsDir))
+		{
+			FileManager.MakeDirectory(*BindingsDir);
+		}
+		
+		static TTuple<const TCHAR*, FText> CommonControllerTypes[] =
+		{
+			MakeTuple(TEXT("vive"), NSLOCTEXT("SteamVR", "CTypeVive", "Vive") ),
+			MakeTuple(TEXT("vive_controller"), NSLOCTEXT("SteamVR", "CTypeViveController", "Vive Controllers") ),
+			MakeTuple(TEXT("oculus_touch"), NSLOCTEXT("SteamVR", "CTypeOculusTouch", "Oculus Touch Controllers") ),
+			MakeTuple(TEXT("holographic_controller"), NSLOCTEXT("SteamVR", "CTypeHolographicController", "Holographic Controllers") ),
+			MakeTuple(TEXT("gamepad"), NSLOCTEXT("SteamVR", "CTypeGamepad", "Game Pads") )
+		};
+
+		for (auto& Item : CommonControllerTypes)
+		{
+			
+			// Skip if the controller type has already be defined
+			if (ExistingBindings.Contains(Item.Key))
+			{
+				continue;
+			}
+
+			// Create a unique file path for the generated file.
+			FString BindingsPath = FileManager.ConvertToAbsolutePathForExternalAppForRead(*FString::Printf(TEXT("%s/%s.json"), *BindingsDir, Item.Key));
+			int count = 0;
+			while (FileManager.FileExists(*BindingsPath) && FileManager.FileSize(*BindingsPath) > 0)
+			{
+				BindingsPath = FileManager.ConvertToAbsolutePathForExternalAppForRead(*FString::Printf(TEXT("%s/%s_%d.json"), *BindingsDir , Item.Key, ++count));
+			}
+
+
+			// Creating a minimal bindings file without any bindings will allow editing it in the SteamVR bindings tool
+			TSharedRef<FJsonObject> BindingsStub = MakeShareable(new FJsonObject());
+			TSharedRef<FJsonObject> EmptyJsonObject = MakeShareable(new FJsonObject());
+			BindingsStub->SetStringField(TEXT("name"), *FText::Format(NSLOCTEXT("SteamVR", "DefaultBindingsFor", "Default bindings for {0}"), Item.Value).ToString());
+			BindingsStub->SetStringField(TEXT("controller_type"), Item.Key);
+
+			// These two fields are required for SteamVR to accept the bindings file and allow editing it.
+			BindingsStub->SetObjectField(TEXT("bindings"), EmptyJsonObject);
+			BindingsStub->SetStringField(TEXT("description"), TEXT(""));
+
+			// Print the stub bindings to a JSON string and save it to a file
+			FString OutputJsonString;
+			TSharedRef<TJsonWriter<>> JsonWriter = TJsonWriterFactory<>::Create(&OutputJsonString);
+			FJsonSerializer::Serialize(BindingsStub, JsonWriter);
+			FFileHelper::SaveStringToFile(OutputJsonString, *BindingsPath, FFileHelper::EEncodingOptions::ForceUTF8WithoutBOM);
+
+			// Add the path of the generated file the action manifest
+			TSharedRef<FJsonObject> BindingObject = MakeShareable(new FJsonObject());
+			BindingObject->SetStringField(TEXT("controller_type"), Item.Key);
+			BindingObject->SetStringField(TEXT("binding_url"), *BindingsPath);
+			InOutDefaultBindings.Add(MakeShareable(new FJsonValueObject(BindingObject)));
+		}
+	}
+#endif
+
 	void BuildActionManifest()
 	{
 		vr::IVRInput* VRInput;
@@ -1377,50 +1451,39 @@ private:
 					UE_LOG(LogSteamVRController, Log, TEXT("Searching for bindings files in %s"), *BindingsDir);
 					for (FString& File : FoundFiles)
 					{
-						FString ControllerType = File;
-						ControllerType.RemoveFromEnd(".json");
+						FString ControllerType;
+						FString JsonStr;
+						FString FilePath = BindingsDir / File;
+						FFileHelper::LoadFileToString(JsonStr, *FilePath);
+						TSharedRef<TJsonReader<TCHAR>> JsonReader = TJsonReaderFactory<TCHAR>::Create(JsonStr);
+						TSharedPtr<FJsonObject> JsonObject = MakeShareable(new FJsonObject());
+						if (!FJsonSerializer::Deserialize(JsonReader, JsonObject) || !JsonObject.IsValid())
+						{
+							UE_LOG(LogSteamVRController, Warning, TEXT("Invalid controller binding file %s: Invalid JSON."), *FilePath);
+							continue;
+						}
+
+						if (!JsonObject->TryGetStringField(TEXT("controller_type"), ControllerType) || ControllerType.IsEmpty())
+						{
+							UE_LOG(LogSteamVRController, Warning, TEXT("Invalid controller binding file %s: Missing or empty controller_type field."), *FilePath);
+							continue;
+						}
+
 						TSharedRef<FJsonObject> BindingObject = MakeShareable(new FJsonObject());
 						BindingObject->SetStringField(TEXT("controller_type"), *ControllerType);
-						BindingObject->SetStringField(TEXT("binding_url"), *FileManager.ConvertToAbsolutePathForExternalAppForRead(*(BindingsDir / File)));
+						BindingObject->SetStringField(TEXT("binding_url"), *FileManager.ConvertToAbsolutePathForExternalAppForRead(*FilePath));
 						DefaultBindings.Add(MakeShareable(new FJsonValueObject(BindingObject)));
 					}
 
-					if (DefaultBindings.Num() == 0)
-					{
 #if WITH_EDITOR
-						UE_LOG(LogSteamVRController, Log, TEXT("No default Steam VR Input bindings found. Adding inital blank entries for common controller types."));
-						
-						// Create the directory if it doesn't exist.
-						if (!FileManager.DirectoryExists(*BindingsDir))
-						{
-							FileManager.MakeDirectory(*BindingsDir);
-						}
-
-						static const TCHAR* CommonControllerTypes[] = 
-						{
-							TEXT("vive"),
-							TEXT("vive_controller"),
-							TEXT("oculus_touch"),
-							TEXT("holographic_controller"),
-							TEXT("gamepad")
-						};
-
-						for (const TCHAR* Type : CommonControllerTypes)
-						{
-							// Telling SteamVR about the path will allow the bindings editor to automatically export the file to the correct location when
-							// pressing "Replace Default Binding".
-							FString BindingsPath = FileManager.ConvertToAbsolutePathForExternalAppForRead(*(BindingsDir / Type + TEXT(".json")));
-							// Creating an empty JSON file will allow editing it in place in the SteamVR Bindings Editor
-							FFileHelper::SaveStringToFile(TEXT(""), *BindingsPath, FFileHelper::EEncodingOptions::ForceUTF8WithoutBOM);
-							TSharedRef<FJsonObject> BindingObject = MakeShareable(new FJsonObject());
-							BindingObject->SetStringField(TEXT("controller_type"), Type);
-							BindingObject->SetStringField(TEXT("binding_url"), *BindingsPath);
-							DefaultBindings.Add(MakeShareable(new FJsonValueObject(BindingObject)));
-						}
+					BuildDefaultActionBindings(BindingsDir, DefaultBindings);
+					check(DefaultBindings.Num());
 #else
+					if (DefaultBindings.Num() == 0)
+					{
 						UE_LOG(LogSteamVRController, Error, TEXT("No default Steam VR Input bindings found in %s."), *BindingsDir);
-#endif
 					}
+#endif
 				}
 
 				TArray<TSharedPtr<FJsonValue>> ActionSets;
