Cherry picking CL#4520709 from Dev-VR
Fix a bug where only some parts of the UE4 console are visible in SteamVR headset even though everything shows up correctly on the spectator screen.

diff --git a/Engine/Plugins/Runtime/Steam/SteamVR/Source/SteamVR/Private/SteamVRHMD.cpp b/Engine/Plugins/Runtime/Steam/SteamVR/Source/SteamVR/Private/SteamVRHMD.cpp
--- a/Engine/Plugins/Runtime/Steam/SteamVR/Source/SteamVR/Private/SteamVRHMD.cpp
+++ b/Engine/Plugins/Runtime/Steam/SteamVR/Source/SteamVR/Private/SteamVRHMD.cpp
@@ -346,13 +346,12 @@ public:
 
 	virtual TSharedPtr< IHeadMountedDisplayVulkanExtensions, ESPMode::ThreadSafe > GetVulkanExtensions() override
 	{
-		FSteamVRHMD* SteamVRHMD = new FSteamVRHMD(this);
+		TSharedPtr< FSteamVRHMD, ESPMode::ThreadSafe > SteamVRHMD = FSceneViewExtensions::NewExtension<FSteamVRHMD>(this);
 		if (!SteamVRHMD->IsInitialized())
 		{
-			delete SteamVRHMD;
 			return nullptr;
 		}
-		return TSharedPtr< IHeadMountedDisplayVulkanExtensions, ESPMode::ThreadSafe >(SteamVRHMD);
+		return SteamVRHMD;
 	}
 
 private:
@@ -368,7 +367,7 @@ IMPLEMENT_MODULE( FSteamVRPlugin, SteamVR )
 TSharedPtr< class IXRTrackingSystem, ESPMode::ThreadSafe > FSteamVRPlugin::CreateTrackingSystem()
 {
 #if STEAMVR_SUPPORTED_PLATFORMS
-	TSharedPtr< FSteamVRHMD, ESPMode::ThreadSafe > SteamVRHMD( new FSteamVRHMD(this) );
+	TSharedPtr< FSteamVRHMD, ESPMode::ThreadSafe > SteamVRHMD = FSceneViewExtensions::NewExtension<FSteamVRHMD>(this);
 	if( SteamVRHMD->IsInitialized() )
 	{
 		VRSystem = SteamVRHMD->GetVRSystem();
@@ -1493,7 +1492,9 @@ bool FSteamVRHMD::NeedReAllocateViewportRenderTarget(const FViewport& Viewport)
 	return false;
 }
 
-FSteamVRHMD::FSteamVRHMD(ISteamVRPlugin* InSteamVRPlugin) :
+FSteamVRHMD::FSteamVRHMD(const FAutoRegister& AutoRegister, ISteamVRPlugin* InSteamVRPlugin) :
+	FHeadMountedDisplayBase(nullptr),
+	FSceneViewExtensionBase(AutoRegister), 
 	bHmdEnabled(true),
 	HmdWornState(EHMDWornState::Unknown),
 	bStereoDesired(false),
diff --git a/Engine/Plugins/Runtime/Steam/SteamVR/Source/SteamVR/Private/SteamVRHMD.h b/Engine/Plugins/Runtime/Steam/SteamVR/Source/SteamVR/Private/SteamVRHMD.h
--- a/Engine/Plugins/Runtime/Steam/SteamVR/Source/SteamVR/Private/SteamVRHMD.h
+++ b/Engine/Plugins/Runtime/Steam/SteamVR/Source/SteamVR/Private/SteamVRHMD.h
@@ -115,7 +115,7 @@ private:
 /**
  * SteamVR Head Mounted Display
  */
-class FSteamVRHMD : public FHeadMountedDisplayBase, public FXRRenderTargetManager, public FSteamVRAssetManager, public TSharedFromThis<FSteamVRHMD, ESPMode::ThreadSafe>, public TStereoLayerManager<FSteamVRLayer>, public IHeadMountedDisplayVulkanExtensions
+class FSteamVRHMD : public FHeadMountedDisplayBase, public FXRRenderTargetManager, public FSteamVRAssetManager, public TStereoLayerManager<FSteamVRLayer>, public FSceneViewExtensionBase, public IHeadMountedDisplayVulkanExtensions
 {
 public:
 	static const FName SteamSystemName;
@@ -135,7 +135,7 @@ public:
 
 	virtual class TSharedPtr< class IStereoRendering, ESPMode::ThreadSafe > GetStereoRenderingDevice() override
 	{
-		return AsShared();
+		return SharedThis(this);
 	}
 
 	virtual bool OnStartGameFrame(FWorldContext& WorldContext) override;
@@ -229,6 +229,16 @@ public:
 	virtual void UpdateSplashScreen() override;
 	virtual void GetAllocatedTexture(uint32 LayerId, FTextureRHIRef &Texture, FTextureRHIRef &LeftTexture) override;
 
+	// ISceneViewExtension interface
+	virtual void SetupViewFamily(FSceneViewFamily& InViewFamily) override {};
+	virtual void SetupView(FSceneViewFamily& InViewFamily, FSceneView& InView) override {}
+	virtual void BeginRenderViewFamily(FSceneViewFamily& InViewFamily) override {}
+	virtual void PreRenderView_RenderThread(FRHICommandListImmediate& RHICmdList, FSceneView& InView) override {}
+	virtual void PreRenderViewFamily_RenderThread(FRHICommandListImmediate& RHICmdList, FSceneViewFamily& InViewFamily) override {}
+	virtual void PostRenderView_RenderThread(FRHICommandListImmediate& RHICmdList, FSceneView& InView) override;
+	virtual bool IsActiveThisFrame(class FViewport* InViewport) const override;
+
+
 	/** IHeadMountedDisplayVulkanExtensions */
 	virtual bool GetVulkanInstanceExtensionsRequired( TArray<const ANSICHAR*>& Out ) override;
 	virtual bool GetVulkanDeviceExtensionsRequired( VkPhysicalDevice_T *pPhysicalDevice, TArray<const ANSICHAR*>& Out ) override;
@@ -361,7 +371,7 @@ public:
 	void PoseToOrientationAndPosition(const vr::HmdMatrix34_t& InPose, const float WorldToMetersScale, FQuat& OutOrientation, FVector& OutPosition) const;
 public:
 	/** Constructor */
-	FSteamVRHMD(ISteamVRPlugin* SteamVRPlugin);
+	FSteamVRHMD(const FAutoRegister&, ISteamVRPlugin*);
 
 	/** Destructor */
 	virtual ~FSteamVRHMD();
diff --git a/Engine/Plugins/Runtime/Steam/SteamVR/Source/SteamVR/Private/SteamVRRender.cpp b/Engine/Plugins/Runtime/Steam/SteamVR/Source/SteamVR/Private/SteamVRRender.cpp
--- a/Engine/Plugins/Runtime/Steam/SteamVR/Source/SteamVR/Private/SteamVRRender.cpp
+++ b/Engine/Plugins/Runtime/Steam/SteamVR/Source/SteamVR/Private/SteamVRRender.cpp
@@ -33,7 +33,6 @@ void FSteamVRHMD::DrawDistortionMesh_RenderThread(struct FRenderingCompositePass
 void FSteamVRHMD::RenderTexture_RenderThread(FRHICommandListImmediate& RHICmdList, FTexture2DRHIParamRef BackBuffer, FTexture2DRHIParamRef SrcTexture, FVector2D WindowSize) const
 {
 	check(IsInRenderingThread());
-	const_cast<FSteamVRHMD*>(this)->UpdateStereoLayers_RenderThread();
 
 	if (bSplashIsShown)
 	{
@@ -45,6 +44,17 @@ void FSteamVRHMD::RenderTexture_RenderThread(FRHICommandListImmediate& RHICmdLis
 	SpectatorScreenController->RenderSpectatorScreen_RenderThread(RHICmdList, BackBuffer, SrcTexture, WindowSize);
 }
 
+void FSteamVRHMD::PostRenderView_RenderThread(FRHICommandListImmediate& RHICmdList, FSceneView& InView)
+{
+	check(IsInRenderingThread());
+	UpdateStereoLayers_RenderThread();
+}
+
+bool FSteamVRHMD::IsActiveThisFrame(class FViewport* InViewport) const
+{
+	return GEngine && GEngine->IsStereoscopic3D(InViewport) && !IsMetalPlatform(GMaxRHIShaderPlatform);
+}
+
 static void DrawOcclusionMesh(FRHICommandList& RHICmdList, EStereoscopicPass StereoPass, const FHMDViewMesh MeshAssets[])
 {
 	check(IsInRenderingThread());
