Unbreak build error caused by bad merge

diff --git a/Engine/Plugins/Runtime/Steam/SteamVR/Source/SteamVR/Private/SteamVRHMD.cpp b/Engine/Plugins/Runtime/Steam/SteamVR/Source/SteamVR/Private/SteamVRHMD.cpp
--- a/Engine/Plugins/Runtime/Steam/SteamVR/Source/SteamVR/Private/SteamVRHMD.cpp
+++ b/Engine/Plugins/Runtime/Steam/SteamVR/Source/SteamVR/Private/SteamVRHMD.cpp
@@ -1493,7 +1493,6 @@ bool FSteamVRHMD::NeedReAllocateViewportRenderTarget(const FViewport& Viewport)
 }
 
 FSteamVRHMD::FSteamVRHMD(const FAutoRegister& AutoRegister, ISteamVRPlugin* InSteamVRPlugin) :
-	FHeadMountedDisplayBase(nullptr),
 	FSceneViewExtensionBase(AutoRegister), 
 	bHmdEnabled(true),
 	HmdWornState(EHMDWornState::Unknown),
