[NO NON-PII COMMIT MESSAGE]

diff --git a/Engine/Plugins/Runtime/Steam/SteamVR/Source/SteamVR/Private/SteamVRStereoLayers.cpp b/Engine/Plugins/Runtime/Steam/SteamVR/Source/SteamVR/Private/SteamVRStereoLayers.cpp
--- a/Engine/Plugins/Runtime/Steam/SteamVR/Source/SteamVR/Private/SteamVRStereoLayers.cpp
+++ b/Engine/Plugins/Runtime/Steam/SteamVR/Source/SteamVR/Private/SteamVRStereoLayers.cpp
@@ -219,10 +219,7 @@ void FSteamVRHMD::UpdateStereoLayers_RenderThread()
 	}
 
 	// Metal is not supported yet
-	if (IsMetalPlatform(GMaxRHIShaderPlatform))
-	{
-		return;
-	}
+	check(!IsMetalPlatform(GMaxRHIShaderPlatform));
 
 	static const auto CVarMixLayerPriorities = IConsoleManager::Get().FindTConsoleVariableDataInt(TEXT("vr.StereoLayers.bMixLayerPriorities"));
 	const bool bUpdateLayerPriorities = (CVarMixLayerPriorities->GetValueOnRenderThread() == 0) && GetStereoLayersDirty();
@@ -378,8 +375,15 @@ void FSteamVRHMD::GetAllocatedTexture(uint32 LayerId, FTextureRHIRef &Texture, F
 }
 
 //=============================================================================
-IStereoLayers* FSteamVRHMD::GetStereoLayers ()
+IStereoLayers* FSteamVRHMD::GetStereoLayers()
 {
+	// Metal is not supported yet. Fall back to the default portable implementation
+	if (IsMetalPlatform(GMaxRHIShaderPlatform))
+	{
+		return FHeadMountedDisplayBase::GetStereoLayers();
+	}
+
+
 	if (VROverlay)
 	{
 		return this;
