Reimplementing CL #4708779 from Release-4.21 (and Dev-VR)

Make SteamVRPlugin responsible for initializing OpenVR

By moving the responsibility we can ensure that the VRCompositor handle is already
available when we need to query Vulkan extensions without needing to initialize SteamVRHMD.

diff --git a/Engine/Plugins/Runtime/Steam/SteamVR/Source/SteamVR/Private/SteamVRHMD.cpp b/Engine/Plugins/Runtime/Steam/SteamVR/Source/SteamVR/Private/SteamVRHMD.cpp
--- a/Engine/Plugins/Runtime/Steam/SteamVR/Source/SteamVR/Private/SteamVRHMD.cpp
+++ b/Engine/Plugins/Runtime/Steam/SteamVR/Source/SteamVR/Private/SteamVRHMD.cpp
@@ -128,7 +128,7 @@ public:
 	}
 #else //STEAMVR_SUPPORTED_PLATFORMS
 		: VRSystem(nullptr)
-		, VRInput(nullptr)
+		, VRCompositor(nullptr)
 	{
 		LoadOpenVRModule();
 	}
@@ -154,7 +154,7 @@ public:
 	
 	virtual vr::IVRInput* GetVRInput() const override
 	{
-		return VRInput;
+		return vr::VRInput();
 	}
 
 	bool LoadOpenVRModule()
@@ -200,14 +200,14 @@ public:
 		}
 		
 		//@todo steamvr: Remove GetProcAddress() workaround once we update to Steamworks 1.33 or higher
-		FSteamVRHMD::VRIsHmdPresentFn = (pVRIsHmdPresent)FPlatformProcess::GetDllExport(OpenVRDLLHandle, TEXT("VR_IsHmdPresent"));
-		FSteamVRHMD::VRGetGenericInterfaceFn = (pVRGetGenericInterface)FPlatformProcess::GetDllExport(OpenVRDLLHandle, TEXT("VR_GetGenericInterface"));
+		VRIsHmdPresentFn = (pVRIsHmdPresent)FPlatformProcess::GetDllExport(OpenVRDLLHandle, TEXT("VR_IsHmdPresent"));
+		VRGetGenericInterfaceFn = (pVRGetGenericInterface)FPlatformProcess::GetDllExport(OpenVRDLLHandle, TEXT("VR_GetGenericInterface"));
 		
 		// Note:  If this fails to compile, it's because you merged a new OpenVR version, and didn't put in the module hacks marked with @epic in openvr.h
-		vr::VR_SetGenericInterfaceCallback(FSteamVRHMD::VRGetGenericInterfaceFn);
+		vr::VR_SetGenericInterfaceCallback(VRGetGenericInterfaceFn);
 
 		// Verify that we've bound correctly to the DLL functions
-		if (!FSteamVRHMD::VRIsHmdPresentFn || !FSteamVRHMD::VRGetGenericInterfaceFn)
+		if (!VRIsHmdPresentFn || !VRGetGenericInterfaceFn)
 		{
 			UE_LOG(LogHMD, Log, TEXT("Failed to GetProcAddress() on openvr_api.dll"));
 			UnloadOpenVRModule();
@@ -217,6 +217,57 @@ public:
 		
 		return true;
 	}
+
+	bool Initialize()
+	{
+		vr::EVRInitError VRInitErr = vr::VRInitError_None;
+		// Attempt to initialize the VRSystem device
+		VRSystem = vr::VR_Init(&VRInitErr, vr::VRApplication_Scene);
+		if ((VRSystem == nullptr) || (VRInitErr != vr::VRInitError_None))
+		{
+			UE_LOG(LogHMD, Log, TEXT("Failed to initialize OpenVR with code %d"), (int32)VRInitErr);
+			return false;
+		}
+
+		// Make sure that the version of the HMD we're compiled against is correct.  This will fill out the proper vtable!
+		VRSystem = (vr::IVRSystem*)(*VRGetGenericInterfaceFn)(vr::IVRSystem_Version, &VRInitErr);
+		if ((VRSystem == nullptr) || (VRInitErr != vr::VRInitError_None))
+		{
+			UE_LOG(LogHMD, Log, TEXT("Failed to initialize OpenVR (version mismatch) with code %d"), (int32)VRInitErr);
+			Reset();
+			return false;
+		}
+
+		// attach to the compositor	
+		int CompositorConnectAttempts = 10;
+		do
+		{
+			VRCompositor = (vr::IVRCompositor*)(*VRGetGenericInterfaceFn)(vr::IVRCompositor_Version, &VRInitErr);
+
+			// If SteamVR was not running when VR_Init was called above, the system may take a few seconds to initialize.
+			// retry a few times before giving up if we get a Compositor connection error.
+			// This is a temporary workaround an issue that will be solved in a future version of SteamVR, where VR_Init will block until everything is ready,
+			// It's only triggered in cases where SteamVR is available, but was not running prior to calling VR_Init above.
+			if ((--CompositorConnectAttempts > 0) && (VRInitErr == vr::VRInitError_IPC_CompositorConnectFailed || VRInitErr == vr::VRInitError_IPC_CompositorInvalidConnectResponse))
+			{
+				UE_LOG(LogHMD, Warning, TEXT("Failed to get Compositor connnection (%d) retrying... (%d attempt(s) left)"), (int32)VRInitErr, CompositorConnectAttempts);
+				FPlatformProcess::Sleep(1);
+			}
+			else
+			{
+				break;
+			}
+		} while (true);
+
+		if (VRInitErr != vr::VRInitError_None)
+		{
+			UE_LOG(LogHMD, Log, TEXT("SteamVR failed to initialize.  Error: %d"), (int32)VRInitErr);
+			Reset();
+			return false;
+		}
+
+		return true;
+	}
 	
 	void UnloadOpenVRModule()
 	{
@@ -230,13 +281,14 @@ public:
 
 	virtual bool IsHMDConnected() override
 	{
-		return FSteamVRHMD::VRIsHmdPresentFn ? (bool)(*FSteamVRHMD::VRIsHmdPresentFn)() : false;
+		return VRIsHmdPresentFn ? (bool)(*VRIsHmdPresentFn)() : false;
 	}
 
 	virtual void Reset() override
 	{
 		VRSystem = nullptr;
-		VRInput = nullptr;
+		VRCompositor = nullptr;
+		vr::VR_Shutdown();
 	}
 	
 	uint64 GetGraphicsAdapterLuid() override
@@ -255,23 +307,11 @@ public:
 #endif
 		
 		{
-			// HACK:  Temporarily stand up the VRSystem to get a graphics adapter.  We're pretty sure we're going to use SteamVR if we get in here, but not guaranteed
-			vr::EVRInitError VRInitErr = vr::VRInitError_None;
-			// Attempt to initialize the VRSystem device
-			vr::IVRSystem* TempVRSystem = vr::VR_Init(&VRInitErr, vr::VRApplication_Scene);
-			if ((TempVRSystem == nullptr) || (VRInitErr != vr::VRInitError_None))
-			{
-				UE_LOG(LogHMD, Log, TEXT("Failed to initialize OpenVR with code %d"), (int32)VRInitErr);
-				return NoDevice;
-			}
-			
-			// Make sure that the version of the HMD we're compiled against is correct.  This will fill out the proper vtable!
-			TempVRSystem = (vr::IVRSystem*)(*FSteamVRHMD::VRGetGenericInterfaceFn)(vr::IVRSystem_Version, &VRInitErr);
-			if ((TempVRSystem == nullptr) || (VRInitErr != vr::VRInitError_None))
+			if (!VRSystem && !Initialize())
 			{
 				return NoDevice;
 			}
-			
+
 			vr::ETextureType TextureType = vr::ETextureType::TextureType_OpenGL;
 #if PLATFORM_MAC
             TextureType = vr::ETextureType::TextureType_IOSurface;
@@ -295,9 +335,7 @@ public:
 			}
 #endif
 
-			TempVRSystem->GetOutputDevice((uint64_t*)((void*)&SelectedDevice), TextureType);
-			
-			vr::VR_Shutdown();
+			VRSystem->GetOutputDevice((uint64_t*)((void*)&SelectedDevice), TextureType);
 		}
 
 
@@ -346,18 +384,27 @@ public:
 
 	virtual TSharedPtr< IHeadMountedDisplayVulkanExtensions, ESPMode::ThreadSafe > GetVulkanExtensions() override
 	{
-		TSharedPtr< FSteamVRHMD, ESPMode::ThreadSafe > SteamVRHMD = FSceneViewExtensions::NewExtension<FSteamVRHMD>(this);
-		if (!SteamVRHMD->IsInitialized())
+#if STEAMVR_SUPPORTED_PLATFORMS && !PLATFORM_MAC
+		if (Initialize())
 		{
-			return nullptr;
+			if (!VulkanExtensions.IsValid())
+			{
+				VulkanExtensions = MakeShareable(new FSteamVRHMD::FVulkanExtensions(VRCompositor));
+			}
+			return VulkanExtensions;
 		}
-		return SteamVRHMD;
+#endif//STEAMVR_SUPPORTED_PLATFORMS
+		return nullptr;
 	}
 
 private:
 	vr::IVRSystem* VRSystem;
-	vr::IVRInput* VRInput;
+	vr::IVRCompositor* VRCompositor;
 
+	pVRIsHmdPresent VRIsHmdPresentFn = nullptr;
+	pVRGetGenericInterface VRGetGenericInterfaceFn = nullptr;
+
+	TSharedPtr< IHeadMountedDisplayVulkanExtensions, ESPMode::ThreadSafe > VulkanExtensions;
 	void* OpenVRDLLHandle;
 #endif // STEAMVR_SUPPORTED_PLATFORMS
 };
@@ -367,11 +414,14 @@ IMPLEMENT_MODULE( FSteamVRPlugin, SteamVR )
 TSharedPtr< class IXRTrackingSystem, ESPMode::ThreadSafe > FSteamVRPlugin::CreateTrackingSystem()
 {
 #if STEAMVR_SUPPORTED_PLATFORMS
+	if (!VRSystem && !Initialize())
+	{
+		return nullptr;
+	}
+
 	TSharedPtr< FSteamVRHMD, ESPMode::ThreadSafe > SteamVRHMD = FSceneViewExtensions::NewExtension<FSteamVRHMD>(this);
 	if( SteamVRHMD->IsInitialized() )
 	{
-		VRSystem = SteamVRHMD->GetVRSystem();
-		VRInput = SteamVRHMD->GetVRInput();
 		return SteamVRHMD;
 	}
 #endif//STEAMVR_SUPPORTED_PLATFORMS
@@ -385,10 +435,14 @@ TSharedPtr< class IXRTrackingSystem, ESPMode::ThreadSafe > FSteamVRPlugin::Creat
 
 #if STEAMVR_SUPPORTED_PLATFORMS
 
-pVRIsHmdPresent FSteamVRHMD::VRIsHmdPresentFn = nullptr;
-pVRGetGenericInterface FSteamVRHMD::VRGetGenericInterfaceFn = nullptr;
+#if !PLATFORM_MAC
+
+FSteamVRHMD::FVulkanExtensions::FVulkanExtensions(vr::IVRCompositor* InVRCompositor)
+	: VRCompositor(InVRCompositor)
+{
+}
 
-bool FSteamVRHMD::GetVulkanInstanceExtensionsRequired(TArray<const ANSICHAR*>& Out)
+bool FSteamVRHMD::FVulkanExtensions::GetVulkanInstanceExtensionsRequired(TArray<const ANSICHAR*>& Out)
 {
 	if (VRCompositor == nullptr)
 	{
@@ -421,7 +475,7 @@ bool FSteamVRHMD::GetVulkanInstanceExtensionsRequired(TArray<const ANSICHAR*>& O
 	return true;
 }
 
-bool FSteamVRHMD::GetVulkanDeviceExtensionsRequired(VkPhysicalDevice_T *pPhysicalDevice, TArray<const ANSICHAR*>& Out)
+bool FSteamVRHMD::FVulkanExtensions::GetVulkanDeviceExtensionsRequired(VkPhysicalDevice_T *pPhysicalDevice, TArray<const ANSICHAR*>& Out)
 {
 	if ( VRCompositor == nullptr )
 	{
@@ -454,9 +508,11 @@ bool FSteamVRHMD::GetVulkanDeviceExtensionsRequired(VkPhysicalDevice_T *pPhysica
 	return true;
 }
 
+#endif
+
 bool FSteamVRHMD::IsHMDConnected()
 {
-	return FSteamVRHMD::VRIsHmdPresentFn ? (bool)(*FSteamVRHMD::VRIsHmdPresentFn)() : false;
+	return SteamVRPlugin->IsHMDConnected();
 }
 
 bool FSteamVRHMD::IsHMDEnabled() const
@@ -1514,10 +1570,12 @@ FSteamVRHMD::FSteamVRHMD(const FAutoRegister& AutoRegister, ISteamVRPlugin* InSt
 	bShouldCheckHMDPosition(false),
 	RendererModule(nullptr),
 	SteamVRPlugin(InSteamVRPlugin),
-	VRSystem(nullptr),
-	VRCompositor(nullptr),
-	VROverlay(nullptr),
-	VRChaperone(nullptr)
+	VRSystem(vr::VRSystem()),
+	VRCompositor(vr::VRCompositor()),
+	VROverlay(vr::VROverlay()),
+	VRChaperone(vr::VRChaperone()),
+	VRRenderModels(vr::VRRenderModels()),
+	VRInput(vr::VRInput())
 {
 	Startup();
 }
@@ -1538,52 +1596,8 @@ bool FSteamVRHMD::Startup()
 	static const FName RendererModuleName("Renderer");
 	RendererModule = FModuleManager::GetModulePtr<IRendererModule>(RendererModuleName);
 
-	vr::EVRInitError VRInitErr = vr::VRInitError_None;
-	// Attempt to initialize the VRSystem device
-	VRSystem = vr::VR_Init(&VRInitErr, vr::VRApplication_Scene);
-	if ((VRSystem == nullptr) || (VRInitErr != vr::VRInitError_None))
-	{
-		UE_LOG(LogHMD, Log, TEXT("Failed to initialize OpenVR with code %d"), (int32)VRInitErr);
-		return false;
-	}
-
-	// Make sure that the version of the HMD we're compiled against is correct.  This will fill out the proper vtable!
-	VRSystem = (vr::IVRSystem*)(*FSteamVRHMD::VRGetGenericInterfaceFn)(vr::IVRSystem_Version, &VRInitErr);
-	if ((VRSystem == nullptr) || (VRInitErr != vr::VRInitError_None))
-	{
-		UE_LOG(LogHMD, Log, TEXT("Failed to initialize OpenVR (version mismatch) with code %d"), (int32)VRInitErr);
-		return false;
-	}
-
-	
-	// attach to the compositor	
-	//VRCompositor = (vr::IVRCompositor*)vr::VR_GetGenericInterface(vr::IVRCompositor_Version, &HmdErr);
-	int CompositorConnectRetries = 10;
-	do
-	{
-		VRCompositor = (vr::IVRCompositor*)(*VRGetGenericInterfaceFn)(vr::IVRCompositor_Version, &VRInitErr);
-		
-		// If SteamVR was not running when VR_Init was called above, the system may take a few seconds to initialize.
-		// retry a few times before giving up if we get a Compositor connection error.
-		// This is a temporary workaround an issue that will be solved in a future version of SteamVR, where VR_Init will block until everything is ready,
-		// It's only triggered in cases where SteamVR is available, but was not running prior to calling VR_Init above.
-		if ((--CompositorConnectRetries > 0) && (VRInitErr == vr::VRInitError_IPC_CompositorConnectFailed || VRInitErr == vr::VRInitError_IPC_CompositorInvalidConnectResponse))
-		{
-			UE_LOG(LogHMD, Warning, TEXT("Failed to get Compositor connnection (%d) retrying... (%d attempt(s) left)"), (int32)VRInitErr, CompositorConnectRetries);
-			FPlatformProcess::Sleep(1);
-		}
-		else
-		{
-			break;
-		}
-	} while (true);
-	
-	if ((VRCompositor != nullptr) && (VRInitErr == vr::VRInitError_None))
-	{
-		VROverlay = (vr::IVROverlay*)(*VRGetGenericInterfaceFn)(vr::IVROverlay_Version, &VRInitErr);
-	}
-
-	if ((VROverlay != nullptr) && (VRInitErr == vr::VRInitError_None))
+	ensure(VRSystem && VRCompositor);
+	if (VRSystem != nullptr && VRCompositor != nullptr)
 	{
 		// grab info about the attached display
 		FString DriverId = GetFStringTrackedDeviceProperty(VRSystem, vr::k_unTrackedDeviceIndex_Hmd, vr::Prop_TrackingSystemName_String);
@@ -1615,27 +1629,11 @@ bool FSteamVRHMD::Startup()
 		CFCFVar->Set(false);
 
 		// Grab the chaperone
-		vr::EVRInitError ChaperoneErr = vr::VRInitError_None;
-		VRChaperone = (vr::IVRChaperone*)(*VRGetGenericInterfaceFn)(vr::IVRChaperone_Version, &ChaperoneErr);
-		if ((VRChaperone != nullptr) && (ChaperoneErr == vr::VRInitError_None))
+		ensure(VRChaperone);
+		if (VRChaperone)
 		{
 			ChaperoneBounds = FChaperoneBounds(VRChaperone);
 		}
-		else
-		{
-			UE_LOG(LogHMD, Log, TEXT("Failed to initialize Chaperone.  Error: %d"), (int32)ChaperoneErr);
-		}
-
-		vr::EVRInitError RenderModelsErr = vr::VRInitError_None;
-		VRRenderModels = (vr::IVRRenderModels*)(*VRGetGenericInterfaceFn)(vr::IVRRenderModels_Version, &RenderModelsErr);
-
-		vr::EVRInitError InputErr = vr::VRInitError_None;
-		VRInput = (vr::IVRInput*)(*VRGetGenericInterfaceFn)(vr::IVRInput_Version, &InputErr);
-		if ((VRInput == nullptr) || (InputErr != vr::VRInitError_None))
-		{
-			UE_LOG(LogHMD, Log, TEXT("Failed to initialize VRInput.  Error: %d"), (int32)InputErr);
-		}
-
 
 #if PLATFORM_MAC
 		if (IsMetalPlatform(GMaxRHIShaderPlatform))
@@ -1678,10 +1676,7 @@ bool FSteamVRHMD::Startup()
 		UE_LOG(LogHMD, Log, TEXT("SteamVR initialized.  Driver: %s  Display: %s"), *DriverId, *DisplayId);
 		return true;
 	}
-	
-	UE_LOG(LogHMD, Log, TEXT("SteamVR failed to initialize.  Err: %d"), (int32)VRInitErr);
 
-	VRSystem = nullptr;
 	return false;
 
 }
@@ -1726,8 +1721,6 @@ void FSteamVRHMD::Shutdown()
 		VRRenderModels = nullptr;
 
 		SteamVRPlugin->Reset();
-
-		vr::VR_Shutdown();
 	}
 }
 
diff --git a/Engine/Plugins/Runtime/Steam/SteamVR/Source/SteamVR/Private/SteamVRHMD.h b/Engine/Plugins/Runtime/Steam/SteamVR/Source/SteamVR/Private/SteamVRHMD.h
--- a/Engine/Plugins/Runtime/Steam/SteamVR/Source/SteamVR/Private/SteamVRHMD.h
+++ b/Engine/Plugins/Runtime/Steam/SteamVR/Source/SteamVR/Private/SteamVRHMD.h
@@ -115,7 +115,7 @@ private:
 /**
  * SteamVR Head Mounted Display
  */
-class FSteamVRHMD : public FHeadMountedDisplayBase, public FXRRenderTargetManager, public FSteamVRAssetManager, public TStereoLayerManager<FSteamVRLayer>, public FSceneViewExtensionBase, public IHeadMountedDisplayVulkanExtensions
+class FSteamVRHMD : public FHeadMountedDisplayBase, public FXRRenderTargetManager, public FSteamVRAssetManager, public TStereoLayerManager<FSteamVRLayer>, public FSceneViewExtensionBase
 {
 public:
 	static const FName SteamSystemName;
@@ -168,7 +168,6 @@ public:
 
 public:
 	/** IHeadMountedDisplay interface */
-
 	virtual bool IsHMDConnected() override;
 	virtual bool IsHMDEnabled() const override;
 	virtual EHMDWornState::Type GetHMDWornState() override;
@@ -238,11 +237,6 @@ public:
 	virtual void PostRenderView_RenderThread(FRHICommandListImmediate& RHICmdList, FSceneView& InView) override;
 	virtual bool IsActiveThisFrame(class FViewport* InViewport) const override;
 
-
-	/** IHeadMountedDisplayVulkanExtensions */
-	virtual bool GetVulkanInstanceExtensionsRequired( TArray<const ANSICHAR*>& Out ) override;
-	virtual bool GetVulkanDeviceExtensionsRequired( VkPhysicalDevice_T *pPhysicalDevice, TArray<const ANSICHAR*>& Out ) override;
-
 	// SpectatorScreen
 private:
 	void CreateSpectatorScreenController();
@@ -313,6 +307,20 @@ public:
 #endif // PLATFORM_WINDOWS
 
 #if !PLATFORM_MAC
+	class FVulkanExtensions : public IHeadMountedDisplayVulkanExtensions
+	{
+	public:
+		FVulkanExtensions(vr::IVRCompositor* InVRCompositor);
+		virtual ~FVulkanExtensions() {}
+
+		/** IHeadMountedDisplayVulkanExtensions */
+		virtual bool GetVulkanInstanceExtensionsRequired(TArray<const ANSICHAR*>& Out) override;
+		virtual bool GetVulkanDeviceExtensionsRequired(VkPhysicalDevice_T *pPhysicalDevice, TArray<const ANSICHAR*>& Out) override;
+
+	private:
+		vr::IVRCompositor* VRCompositor;
+	};
+
 	class VulkanBridge : public BridgeBaseImpl
 	{
 	public:
@@ -593,9 +601,6 @@ private:
 
 //@todo steamvr: Remove GetProcAddress() workaround once we have updated to Steamworks 1.33 or higher
 public:
-	static pVRIsHmdPresent VRIsHmdPresentFn;
-	static pVRGetGenericInterface VRGetGenericInterfaceFn;
-
 	friend class FSteamSplashTicker;
 };
 
diff --git a/Engine/Plugins/Runtime/Steam/SteamVR/Source/SteamVR/Private/SteamVRStereoLayers.cpp b/Engine/Plugins/Runtime/Steam/SteamVR/Source/SteamVR/Private/SteamVRStereoLayers.cpp
--- a/Engine/Plugins/Runtime/Steam/SteamVR/Source/SteamVR/Private/SteamVRStereoLayers.cpp
+++ b/Engine/Plugins/Runtime/Steam/SteamVR/Source/SteamVR/Private/SteamVRStereoLayers.cpp
@@ -383,7 +383,7 @@ IStereoLayers* FSteamVRHMD::GetStereoLayers()
 		return FHeadMountedDisplayBase::GetStereoLayers();
 	}
 
-
+	ensure(VROverlay);
 	if (VROverlay)
 	{
 		return this;
