UE-62507 Fatal error crash opening packaged QAGame
UE-62544 //UE4/Dev-VR - Run Automated Tests Cooked Win64 - Unhandled Exception
-Updated #define OPENVR_SDK_VER TEXT("OpenVRv1_0_16")

diff --git a/Engine/Plugins/Runtime/Steam/SteamVR/Source/SteamVR/Private/SteamVRPrivate.h b/Engine/Plugins/Runtime/Steam/SteamVR/Source/SteamVR/Private/SteamVRPrivate.h
--- a/Engine/Plugins/Runtime/Steam/SteamVR/Source/SteamVR/Private/SteamVRPrivate.h
+++ b/Engine/Plugins/Runtime/Steam/SteamVR/Source/SteamVR/Private/SteamVRPrivate.h
@@ -3,7 +3,7 @@
 #pragma once
 
 /** Name of the current OpenVR SDK version in use (matches directory name) */
-#define OPENVR_SDK_VER TEXT("OpenVRv1_0_15")
+#define OPENVR_SDK_VER TEXT("OpenVRv1_0_16")
 
 // for STEAMVR_SUPPORTED_PLATFORMS, keep at top
 #include "ISteamVRPlugin.h"
