Adding TrackedDeviceClass_DisplayRedirect as a known steam vr tracked device type.

diff --git a/Engine/Plugins/Runtime/Steam/SteamVR/Source/SteamVRController/Private/SteamVRController.cpp b/Engine/Plugins/Runtime/Steam/SteamVR/Source/SteamVRController/Private/SteamVRController.cpp
--- a/Engine/Plugins/Runtime/Steam/SteamVR/Source/SteamVRController/Private/SteamVRController.cpp
+++ b/Engine/Plugins/Runtime/Steam/SteamVR/Source/SteamVRController/Private/SteamVRController.cpp
@@ -926,6 +926,8 @@ private:
 			case vr::TrackedDeviceClass_HMD:
 				// falls through
 			case vr::TrackedDeviceClass_TrackingReference:
+				// falls through
+			case vr::TrackedDeviceClass_DisplayRedirect:
 				break;
 			default:
 				UE_LOG(LogSteamVRController, Warning, TEXT("Encountered unsupported device class of %i!"), (int32)DeviceClass);
