Set the graphics pipeline state before setting the shader parameters

diff --git a/Engine/Plugins/Runtime/Steam/SteamVR/Source/SteamVR/Private/SteamVRSpectatorScreen.cpp b/Engine/Plugins/Runtime/Steam/SteamVR/Source/SteamVR/Private/SteamVRSpectatorScreen.cpp
--- a/Engine/Plugins/Runtime/Steam/SteamVR/Source/SteamVR/Private/SteamVRSpectatorScreen.cpp
+++ b/Engine/Plugins/Runtime/Steam/SteamVR/Source/SteamVR/Private/SteamVRSpectatorScreen.cpp
@@ -87,17 +87,19 @@ void FSteamVRHMD::CopyTexture_RenderThread(FRHICommandListImmediate& RHICmdList,
 		{
 			TShaderMapRef<FScreenPSsRGBSource> PixelShader(ShaderMap);
 			GraphicsPSOInit.BoundShaderState.PixelShaderRHI = GETSAFERHISHADER_PIXEL(*PixelShader);
+
+			SetGraphicsPipelineState(RHICmdList, GraphicsPSOInit);
 			PixelShader->SetParameters(RHICmdList, PixelSampler, SrcTexture);
 		}
 		else
 		{
 			TShaderMapRef<FScreenPS> PixelShader(ShaderMap);
 			GraphicsPSOInit.BoundShaderState.PixelShaderRHI = GETSAFERHISHADER_PIXEL(*PixelShader);
+
+			SetGraphicsPipelineState(RHICmdList, GraphicsPSOInit);
 			PixelShader->SetParameters(RHICmdList, PixelSampler, SrcTexture);
 		}
 
-		SetGraphicsPipelineState(RHICmdList, GraphicsPSOInit);
-
 		RendererModule->DrawRectangle(
 			RHICmdList,
 			0, 0,
