Convert more parts of the SteamVR plugin to use the OpenVR module context
To avoid a possible crash we don't rely on the module context for the compositor
Add support for re-initializing the context as well in case of a SteamVR shutdown

diff --git a/Engine/Plugins/Runtime/Steam/SteamVR/Source/SteamVR/Private/SteamVRAssetManager.cpp b/Engine/Plugins/Runtime/Steam/SteamVR/Source/SteamVR/Private/SteamVRAssetManager.cpp
--- a/Engine/Plugins/Runtime/Steam/SteamVR/Source/SteamVR/Private/SteamVRAssetManager.cpp
+++ b/Engine/Plugins/Runtime/Steam/SteamVR/Source/SteamVR/Private/SteamVRAssetManager.cpp
@@ -49,34 +49,31 @@ static int32 SteamVRDevice_Impl::GetDeviceStringProperty(int32 DeviceIndex, int3
 {
 	int32 ErrorResult = -1;
 #if STEAMVR_SUPPORTED_PLATFORMS
-	if (FSteamVRHMD* SteamHMD = GetSteamHMD())
+	vr::IVRSystem* SteamVRSystem = vr::VRSystem();
+	if (SteamVRSystem)
 	{
-		vr::IVRSystem* SteamVRSystem = SteamHMD->GetVRSystem();
-		if (SteamVRSystem)
-		{
-			vr::ETrackedDeviceProperty SteamPropId = (vr::ETrackedDeviceProperty)PropertyId;
+		vr::ETrackedDeviceProperty SteamPropId = (vr::ETrackedDeviceProperty)PropertyId;
 
-			vr::TrackedPropertyError APIError;
-			TArray<char> Buffer;
-			Buffer.AddUninitialized(vr::k_unMaxPropertyStringSize);
+		vr::TrackedPropertyError APIError;
+		TArray<char> Buffer;
+		Buffer.AddUninitialized(vr::k_unMaxPropertyStringSize);
 
-			int Size = SteamVRSystem->GetStringTrackedDeviceProperty(DeviceIndex, SteamPropId, Buffer.GetData(), Buffer.Num(), &APIError);
-			if (APIError == vr::TrackedProp_BufferTooSmall)
-			{
-				Buffer.AddUninitialized(Size - Buffer.Num());
-				Size = SteamVRSystem->GetStringTrackedDeviceProperty(DeviceIndex, SteamPropId, Buffer.GetData(), Buffer.Num(), &APIError);
-			}
+		int Size = SteamVRSystem->GetStringTrackedDeviceProperty(DeviceIndex, SteamPropId, Buffer.GetData(), Buffer.Num(), &APIError);
+		if (APIError == vr::TrackedProp_BufferTooSmall)
+		{
+			Buffer.AddUninitialized(Size - Buffer.Num());
+			Size = SteamVRSystem->GetStringTrackedDeviceProperty(DeviceIndex, SteamPropId, Buffer.GetData(), Buffer.Num(), &APIError);
+		}
 
-			if (APIError == vr::TrackedProp_Success)
-			{
-				StringPropertyOut = UTF8_TO_TCHAR(Buffer.GetData());
-			}
-			else
-			{
-				StringPropertyOut = UTF8_TO_TCHAR(SteamVRSystem->GetPropErrorNameFromEnum(APIError));
-			}
-			ErrorResult = (int32)APIError;
-		}	
+		if (APIError == vr::TrackedProp_Success)
+		{
+			StringPropertyOut = UTF8_TO_TCHAR(Buffer.GetData());
+		}
+		else
+		{
+			StringPropertyOut = UTF8_TO_TCHAR(SteamVRSystem->GetPropErrorNameFromEnum(APIError));
+		}
+		ErrorResult = (int32)APIError;
 	}
 #endif 
 	return ErrorResult;
@@ -86,10 +83,7 @@ static vr::IVRRenderModels* SteamVRDevice_Impl::GetSteamVRModelManager()
 {
 	vr::IVRRenderModels* VRModelManager = nullptr;
 #if STEAMVR_SUPPORTED_PLATFORMS
-	if (FSteamVRHMD* SteamHMD = SteamVRDevice_Impl::GetSteamHMD())
-	{
-		VRModelManager = SteamHMD->GetRenderModelManager();
-	}
+	VRModelManager = vr::VRRenderModels();
 #endif
 	return VRModelManager;
 }
@@ -618,7 +612,7 @@ bool FSteamVRAssetManager::EnumerateRenderableDevices(TArray<int32>& DeviceListO
 
 #if STEAMVR_SUPPORTED_PLATFORMS
 	FSteamVRHMD* SteamHMD = SteamVRDevice_Impl::GetSteamHMD();
-	bHasActiveVRSystem = SteamHMD && SteamHMD->GetVRSystem();
+	bHasActiveVRSystem = SteamHMD && SteamHMD->IsInitialized();
 	
 	if (bHasActiveVRSystem)
 	{
@@ -642,8 +636,7 @@ int32 FSteamVRAssetManager::GetDeviceId(EControllerHand ControllerHand)
 	int32 DeviceIndexOut = INDEX_NONE;
 
 #if STEAMVR_SUPPORTED_PLATFORMS
-	FSteamVRHMD* SteamHMD = SteamVRDevice_Impl::GetSteamHMD();
-	vr::IVRSystem* SteamVRSystem = (SteamHMD) ? SteamHMD->GetVRSystem() : nullptr;
+	vr::IVRSystem* SteamVRSystem = vr::VRSystem();
 	if (SteamVRSystem)
 	{
 		vr::ETrackedDeviceClass DesiredDeviceClass = vr::ETrackedDeviceClass::TrackedDeviceClass_Invalid;
@@ -733,8 +726,7 @@ UPrimitiveComponent* FSteamVRAssetManager::CreateRenderComponent(const int32 Dev
 	FString ModelName;
 	if (SteamVRDevice_Impl::GetDeviceStringProperty(DeviceId, vr::Prop_RenderModelName_String, ModelName) == vr::TrackedProp_Success)
 	{
-		FSteamVRHMD* SteamHMD = SteamVRDevice_Impl::GetSteamHMD();
-		vr::IVRRenderModels* VRModelManager = (SteamHMD) ? SteamHMD->GetRenderModelManager() : nullptr;
+		vr::IVRRenderModels* VRModelManager = vr::VRRenderModels();
 
 		if (VRModelManager != nullptr)
 		{
diff --git a/Engine/Plugins/Runtime/Steam/SteamVR/Source/SteamVR/Private/SteamVRHMD.cpp b/Engine/Plugins/Runtime/Steam/SteamVR/Source/SteamVR/Private/SteamVRHMD.cpp
--- a/Engine/Plugins/Runtime/Steam/SteamVR/Source/SteamVR/Private/SteamVRHMD.cpp
+++ b/Engine/Plugins/Runtime/Steam/SteamVR/Source/SteamVR/Private/SteamVRHMD.cpp
@@ -152,9 +152,9 @@ public:
 		return VRSystem;
 	}
 	
-	virtual vr::IVRInput* GetVRInput() const override
+	virtual vr::IVRCompositor* GetVRCompositor() const override
 	{
-		return vr::VRInput();
+		return VRCompositor;
 	}
 
 	bool LoadOpenVRModule()
@@ -653,7 +653,7 @@ bool FSteamVRHMD::GetCurrentPose(int32 DeviceId, FQuat& CurrentOrientation, FVec
 
 void FSteamVRHMD::UpdatePoses()
 {
-	if (VRSystem == nullptr)
+	if (!VRSystem || !VRCompositor)
 	{
 		return;
 	}
@@ -1571,12 +1571,10 @@ FSteamVRHMD::FSteamVRHMD(const FAutoRegister& AutoRegister, ISteamVRPlugin* InSt
 	bShouldCheckHMDPosition(false),
 	RendererModule(nullptr),
 	SteamVRPlugin(InSteamVRPlugin),
-	VRSystem(vr::VRSystem()),
-	VRCompositor(vr::VRCompositor()),
+	VRSystem(InSteamVRPlugin->GetVRSystem()),
+	VRCompositor(InSteamVRPlugin->GetVRCompositor()),
 	VROverlay(vr::VROverlay()),
-	VRChaperone(vr::VRChaperone()),
-	VRRenderModels(vr::VRRenderModels()),
-	VRInput(vr::VRInput())
+	VRChaperone(vr::VRChaperone())
 {
 	Startup();
 }
@@ -1588,7 +1586,7 @@ FSteamVRHMD::~FSteamVRHMD()
 
 bool FSteamVRHMD::IsInitialized() const
 {
-	return (VRSystem != nullptr);
+	return (VRSystem != nullptr) && (VRCompositor != nullptr);
 }
 
 bool FSteamVRHMD::Startup()
@@ -1597,8 +1595,17 @@ bool FSteamVRHMD::Startup()
 	static const FName RendererModuleName("Renderer");
 	RendererModule = FModuleManager::GetModulePtr<IRendererModule>(RendererModuleName);
 
-	ensure(VRSystem && VRCompositor);
-	if (VRSystem != nullptr && VRCompositor != nullptr)
+	// Re-initialize the plugin if we're canceling the shutdown
+	if (!IsInitialized())
+	{
+		SteamVRPlugin->Initialize();
+		VRSystem = SteamVRPlugin->GetVRSystem();
+		VRCompositor = SteamVRPlugin->GetVRCompositor();
+		VROverlay = vr::VROverlay();
+		VRChaperone = vr::VRChaperone();
+	}
+
+	if (ensure(IsInitialized()))
 	{
 		// grab info about the attached display
 		FString DriverId = GetFStringTrackedDeviceProperty(VRSystem, vr::k_unTrackedDeviceIndex_Hmd, vr::Prop_TrackingSystemName_String);
@@ -1717,9 +1724,7 @@ void FSteamVRHMD::Shutdown()
 		// shut down our headset
 		VRSystem = nullptr;
 		VRCompositor = nullptr;
-		VROverlay = nullptr;
 		VRChaperone = nullptr;
-		VRRenderModels = nullptr;
 
 		SteamVRPlugin->Reset();
 	}
diff --git a/Engine/Plugins/Runtime/Steam/SteamVR/Source/SteamVR/Private/SteamVRHMD.h b/Engine/Plugins/Runtime/Steam/SteamVR/Source/SteamVR/Private/SteamVRHMD.h
--- a/Engine/Plugins/Runtime/Steam/SteamVR/Source/SteamVR/Private/SteamVRHMD.h
+++ b/Engine/Plugins/Runtime/Steam/SteamVR/Source/SteamVR/Private/SteamVRHMD.h
@@ -388,10 +388,6 @@ public:
 	/** @return	True if the API was initialized OK */
 	bool IsInitialized() const;
 
-	vr::IVRSystem* GetVRSystem() const { return VRSystem; }
-	vr::IVRInput* GetVRInput() const { return VRInput; }
-	vr::IVRRenderModels* GetRenderModelManager() const { return VRRenderModels; }
-
 protected:
 
 	virtual float GetWorldToMetersScale() const override;
@@ -590,8 +586,6 @@ private:
 	vr::IVRCompositor* VRCompositor;
 	vr::IVROverlay* VROverlay;
 	vr::IVRChaperone* VRChaperone;
-	vr::IVRRenderModels* VRRenderModels;
-	vr::IVRInput* VRInput;
 
 	FString DisplayId;
 
diff --git a/Engine/Plugins/Runtime/Steam/SteamVR/Source/SteamVR/Public/ISteamVRPlugin.h b/Engine/Plugins/Runtime/Steam/SteamVR/Source/SteamVR/Public/ISteamVRPlugin.h
--- a/Engine/Plugins/Runtime/Steam/SteamVR/Source/SteamVR/Public/ISteamVRPlugin.h
+++ b/Engine/Plugins/Runtime/Steam/SteamVR/Source/SteamVR/Public/ISteamVRPlugin.h
@@ -56,11 +56,16 @@ public:
 	virtual vr::IVRSystem* GetVRSystem() const = 0;
 
 	/**
-	* Get the IVRInput* that was previously set by the HMD implemenentation.
+	* Get the IVRCompositor* that was previously set by the HMD implemenentation.
 	*
 	* @return The pointer if the HMD has been initialized, otherwise nullptr
 	*/
-	virtual vr::IVRInput* GetVRInput() const = 0;
+	virtual vr::IVRCompositor* GetVRCompositor() const = 0;
+
+	/**
+	 * Initializes the plugin, in case of VR system restart
+	 */
+	virtual bool Initialize() = 0;
 
 	/**
 	 * Resets the plugin, in case of early VR system shutdown
diff --git a/Engine/Plugins/Runtime/Steam/SteamVR/Source/SteamVRController/Private/SteamVRController.cpp b/Engine/Plugins/Runtime/Steam/SteamVR/Source/SteamVRController/Private/SteamVRController.cpp
--- a/Engine/Plugins/Runtime/Steam/SteamVR/Source/SteamVRController/Private/SteamVRController.cpp
+++ b/Engine/Plugins/Runtime/Steam/SteamVR/Source/SteamVRController/Private/SteamVRController.cpp
@@ -529,7 +529,7 @@ public:
 #if STEAMVRCONTROLLER_SUPPORTED_PLATFORMS
 	void SendActionInputEvents()
 	{
-		vr::IVRInput* VRInput = GetVRInput();
+		vr::IVRInput* VRInput = vr::VRInput();
 
 		if (VRInput != nullptr)
 		{
@@ -856,16 +856,6 @@ public:
 	static ESteamVRTouchDPadMapping DefaultDPadMapping;
 private:
 
-	inline vr::IVRInput* GetVRInput() const
-	{
-		if (SteamVRPlugin == nullptr)
-		{
-			SteamVRPlugin = &FModuleManager::LoadModuleChecked<ISteamVRPlugin>(TEXT("SteamVR"));
-		}
-
-		return SteamVRPlugin->GetVRInput();
-	}
-
 	inline vr::IVRSystem* GetVRSystem() const
 	{
 		if (SteamVRPlugin == nullptr)
@@ -1306,7 +1296,7 @@ private:
 	{
 		vr::IVRInput* VRInput;
 
-		if (bEnableVRInput && (VRInput = GetVRInput()) != nullptr)
+		if (bEnableVRInput && (VRInput = vr::VRInput()) != nullptr)
 		{
 			Actions.Empty();
 			auto InputSettings = GetDefault<UInputSettings>();
