Fix SteamVR plugin controller latency issues. Turns out the SteamVR plugin was calling the incorrect API to get the poses for the HMD and controllers; this was causing jittery movement and also some major artifacting with Temporal AA.

diff --git a/Engine/Plugins/Runtime/Steam/SteamVR/Source/SteamVR/Private/SteamVRHMD.cpp b/Engine/Plugins/Runtime/Steam/SteamVR/Source/SteamVR/Private/SteamVRHMD.cpp
--- a/Engine/Plugins/Runtime/Steam/SteamVR/Source/SteamVR/Private/SteamVRHMD.cpp
+++ b/Engine/Plugins/Runtime/Steam/SteamVR/Source/SteamVR/Private/SteamVRHMD.cpp
@@ -669,7 +669,7 @@ void FSteamVRHMD::UpdatePoses()
 	else
 	{
 		check(IsInGameThread());
-		VRSystem->GetDeviceToAbsoluteTrackingPose(VRCompositor->GetTrackingSpace(), 0.0f, Poses, ARRAYSIZE(Poses));
+		vr::EVRCompositorError PoseError = VRCompositor->GetLastPoses(Poses, ARRAYSIZE(Poses), NULL, 0);
 	}
 
 	TrackingFrame.bHaveVisionTracking = false;
