Copy of CL#7473960 from Dev-VR:
Store SteamVR resources in TSharedPtrs instead of rolling our own ref counting.
Simplify code for loading resources.
More tentative fixes for #jira UE-57680

diff --git a/Engine/Plugins/Runtime/Steam/SteamVR/Source/SteamVR/Private/SteamVRAssetManager.cpp b/Engine/Plugins/Runtime/Steam/SteamVR/Source/SteamVR/Private/SteamVRAssetManager.cpp
--- a/Engine/Plugins/Runtime/Steam/SteamVR/Source/SteamVR/Private/SteamVRAssetManager.cpp
+++ b/Engine/Plugins/Runtime/Steam/SteamVR/Source/SteamVR/Private/SteamVRAssetManager.cpp
@@ -30,7 +30,6 @@ namespace SteamVRDevice_Impl
 {
 	static FSteamVRHMD* GetSteamHMD();
 	static int32 GetDeviceStringProperty(int32 DeviceIndex, int32 PropertyId, FString& StringPropertyOut);
-	static vr::IVRRenderModels* GetSteamVRModelManager();
 }
 
 static FSteamVRHMD* SteamVRDevice_Impl::GetSteamHMD()
@@ -79,193 +78,217 @@ static int32 SteamVRDevice_Impl::GetDeviceStringProperty(int32 DeviceIndex, int3
 	return ErrorResult;
 }
 
-static vr::IVRRenderModels* SteamVRDevice_Impl::GetSteamVRModelManager()
-{
-	vr::IVRRenderModels* VRModelManager = nullptr;
-#if STEAMVR_SUPPORTED_PLATFORMS
-	VRModelManager = vr::VRRenderModels();
-#endif
-	return VRModelManager;
-}
 
-/* TSteamVRResource 
+/* TSteamVRResourceLoader
  *****************************************************************************/
-
-template<typename ResType>
-struct TSharedSteamVRResource
-{
-	TSharedSteamVRResource() : RefCount(0), RawResource(nullptr) {}
-	int32    RefCount;
-	ResType* RawResource;
-};
-
-template< typename ResType, typename IDType >
-struct TSteamVRResource
+template<typename ResType, typename IDType>
+class TSteamVRResourceLoader
 {
-private:
-	typedef TSharedSteamVRResource<ResType>  FSharedSteamVRResource;
-	typedef TMap<IDType, FSharedSteamVRResource> TSharedResourceMap;
-	static  TSharedResourceMap SharedResources;
-
 public:
-	TSteamVRResource(const IDType& ResID, bool bKickOffLoad = true)
-		: ResourceId(ResID)
-		, RawResource(nullptr)
-		, bLoadFailed(false)
+	TSteamVRResourceLoader() {}
+
+	bool LoadResource_Async(const IDType& ResID, TSharedPtr<ResType>& ResourceOut)
 	{
-		SharedResources.FindOrAdd(ResourceId).RefCount += 1;
-		if (bKickOffLoad)
+#if STEAMVR_SUPPORTED_PLATFORMS
+		vr::IVRRenderModels* VRModelManager = GetSteamVRModelManager();
+		if (VRModelManager)
 		{
-			TickAsyncLoad();
+			ResType* RawResourcePtr;
+			vr::EVRRenderModelError err = LoadResource_Internal(VRModelManager, ResID, &RawResourcePtr);
+			if (err == vr::VRRenderModelError_None)
+			{
+				ResourceOut = WrapResource(RawResourcePtr);
+				return true;
+			}
+			else
+			{
+				return err == vr::VRRenderModelError_Loading;
+			}
 		}
+#endif
+		return false;
 	}
 
-	bool IsPending()
-	{
-		return (RawResource == nullptr) && !bLoadFailed;
-	}
-
-	bool IsValid()
+private:
+	static vr::IVRRenderModels* GetSteamVRModelManager()
 	{
-		return (RawResource != nullptr);
+		vr::IVRRenderModels* VRModelManager = nullptr;
+#if STEAMVR_SUPPORTED_PLATFORMS
+		VRModelManager = vr::VRRenderModels();
+#endif
+		return VRModelManager;
 	}
 
-	ResType* TickAsyncLoad()
+	// When loading the same resource id multiple times, SteamVR will keep returning the same pointer until it is freed. 
+	// We wrap it in a TSharedPtr and save a weak reference to it so we can deallocate it when we no longer have 
+	// references to it.	
+	TSharedPtr<ResType> WrapResource(ResType* RawResourcePtr)
 	{
-		if (IsPending())
+		if (!RawResourcePtr)
 		{
-			FSharedSteamVRResource& SharedResource = SharedResources[ResourceId]; // The shared resource was already created in the constructor
-			if (SharedResource.RawResource != nullptr)
+			return TSharedPtr<ResType>();
+		}
+		else
+		{
+			// First check if we already have created a shared pointer for this resource
+			TWeakPtr<ResType>* Existing = ActiveResources.Find(RawResourcePtr);
+			if (Existing)
 			{
-				RawResource = SharedResource.RawResource;
+				// The custom deleter will remove the value from the map before the weak pointer becomes invalid,
+				// and it should point to the same object as the naked pointer passed in.
+				check(Existing->IsValid() && Existing->HasSameObject(RawResourcePtr));
+				UE_LOG(LogSteamVR, Log, TEXT("TSteamVRResourceLoader::WrapResource: Returning existing resource %p"), RawResourcePtr);
+				return Existing->Pin();
 			}
-#if STEAMVR_SUPPORTED_PLATFORMS
-			else if (vr::IVRRenderModels* VRModelManager = SteamVRDevice_Impl::GetSteamVRModelManager())
-			{
-				vr::EVRRenderModelError LoadResult = (vr::EVRRenderModelError)TickAsyncLoad_Internal(VRModelManager, &RawResource);
-
-				const bool bLoadComplete = (LoadResult != vr::VRRenderModelError_Loading);
-				if (bLoadComplete)
-				{
-					bLoadFailed = !RawResource || (LoadResult != vr::VRRenderModelError_None);
-					if (!bLoadFailed)
-					{
-						SharedResource.RawResource = RawResource;
-					}
-					else
-					{
-						RawResource = nullptr;
-					}
-				}
-			}
-#endif
 			else
 			{
-				bLoadFailed = true;
+				UE_LOG(LogSteamVR, Log, TEXT("TSteamVRResourceLoader::WrapResource: Creating new TSharedPtr for resource %p"), RawResourcePtr);
+				// Construct a new shared pointer with a custom deleter.
+				TSharedPtr<ResType> NewPtr = MakeShareable(RawResourcePtr, [this](ResType* Instance)
+				{
+					FreeResource(Instance);
+				});
+				// Save a weak reference so we only construct a single shared pointer for each unique resource.
+				ActiveResources.Emplace(RawResourcePtr, NewPtr);
+				return NewPtr;
 			}
 		}
-		return RawResource;
 	}
 
-	operator ResType*()   { return RawResource; }
-	ResType* operator->() { return RawResource; }
-
-	IDType GetId() const  { return ResourceId;  }
-
-protected:
-	int32 TickAsyncLoad_Internal(vr::IVRRenderModels* VRModelManager, ResType** ResourceOut);
-
-	void Reset()
+	void FreeResource(ResType* RawResourcePtr)
 	{
-		if (FSharedSteamVRResource* SharedResource = SharedResources.Find(ResourceId))
+		vr::IVRRenderModels* VRModelManager = GetSteamVRModelManager();
+		if (VRModelManager)
 		{
-			int32 NewRefCount = --SharedResource->RefCount;
-			if (NewRefCount <= 0)
-			{
-				if ( vr::IVRRenderModels* VRModelManager = SteamVRDevice_Impl::GetSteamVRModelManager() )
-				{
-					if (!RawResource)
-					{
-						RawResource = SharedResource->RawResource;
-					}
-					FreeResource(VRModelManager);
-				}
-				SharedResources.Remove(ResourceId);
-			}
+			FreeResource_Internal(VRModelManager, RawResourcePtr);
+		}
+		else
+		{
+			UE_LOG(LogSteamVR, Warning, TEXT("FSteamVRModelLoader::FreeResource: Could not get SteamVR model manager when freeing render model instance %p."), RawResourcePtr);
 		}
-		RawResource = nullptr;
+		ActiveResources.Remove(RawResourcePtr);
 	}
 
-	void FreeResource(vr::IVRRenderModels* VRModelManager);
-
-protected:
-	IDType   ResourceId;
-	ResType* RawResource;
-	bool     bLoadFailed;
-};
+#if STEAMVR_SUPPORTED_PLATFORMS
+	static vr::EVRRenderModelError LoadResource_Internal(vr::IVRRenderModels* VRModelManager, const IDType& ResID, ResType** ResourceOut);
+#endif
+	static void FreeResource_Internal(vr::IVRRenderModels* VRModelManager, ResType* RawResourcePtr);
 
-template<typename ResType, typename IDType>
-TMap< IDType, TSharedSteamVRResource<ResType> > TSteamVRResource<ResType, IDType>::SharedResources;
 
-typedef TSteamVRResource<vr::RenderModel_t, FString> TSteamVRModel;
-typedef TSteamVRResource<vr::RenderModel_TextureMap_t, int32> TSteamVRTexture;
+	TMap<ResType*, TWeakPtr<ResType>> ActiveResources;
+};
+typedef TSteamVRResourceLoader<vr::RenderModel_t, FString> FSteamVRModelLoader;
+typedef TSteamVRResourceLoader<vr::RenderModel_TextureMap_t, int32> FSteamVRTextureLoader;
 
 /// @cond DOXYGEN_WARNINGS
+#if STEAMVR_SUPPORTED_PLATFORMS
 template<>
-int32 TSteamVRModel::TickAsyncLoad_Internal(vr::IVRRenderModels* VRModelManager, vr::RenderModel_t** ResourceOut)
+vr::EVRRenderModelError FSteamVRModelLoader::LoadResource_Internal(vr::IVRRenderModels* VRModelManager, const FString& ResourceId, vr::RenderModel_t** RawResourceOut)
 {
-#if STEAMVR_SUPPORTED_PLATFORMS
-	return (int32)VRModelManager->LoadRenderModel_Async(TCHAR_TO_UTF8(*ResourceId), ResourceOut);
-#else
-	return INDEX_NONE;
-#endif
+	return VRModelManager->LoadRenderModel_Async(TCHAR_TO_UTF8(*ResourceId), RawResourceOut);
 }
+#endif
 /// @endcond
 
 template<>
-void TSteamVRModel::FreeResource(vr::IVRRenderModels* VRModelManager)
+void FSteamVRModelLoader::FreeResource_Internal(vr::IVRRenderModels* VRModelManager, vr::RenderModel_t* RawResourcePtr)
 {
 #if STEAMVR_SUPPORTED_PLATFORMS
-	UE_LOG(LogSteamVR, Log, TEXT("TSteamVRModel::FreeResource: Freeing render model instance. ResourceId: %s"), *ResourceId);
-	VRModelManager->FreeRenderModel(RawResource);
+	check(VRModelManager);
+	UE_LOG(LogSteamVR, Log, TEXT("FSteamVRModelLoader::FreeResource_Internal: Freeing render model instance %p"), RawResourcePtr);
+	VRModelManager->FreeRenderModel(RawResourcePtr);
 #endif
 }
 
 /// @cond DOXYGEN_WARNINGS
+#if STEAMVR_SUPPORTED_PLATFORMS
 template<>
-int32 TSteamVRTexture::TickAsyncLoad_Internal(vr::IVRRenderModels* VRModelManager, vr::RenderModel_TextureMap_t** ResourceOut)
+vr::EVRRenderModelError FSteamVRTextureLoader::LoadResource_Internal(vr::IVRRenderModels* VRModelManager, const int32& ResourceId, vr::RenderModel_TextureMap_t** RawResourceOut)
 {
-#if STEAMVR_SUPPORTED_PLATFORMS
-	return (int32)VRModelManager->LoadTexture_Async(ResourceId, ResourceOut);
-#else
-	return INDEX_NONE;
-#endif
+	return VRModelManager->LoadTexture_Async(ResourceId, RawResourceOut);
 }
+#endif
 /// @endcond
 
 template<>
-void TSteamVRTexture::FreeResource(vr::IVRRenderModels* VRModelManager)
+void FSteamVRTextureLoader::FreeResource_Internal(vr::IVRRenderModels* VRModelManager, vr::RenderModel_TextureMap_t* RawResourcePtr)
 {
 #if STEAMVR_SUPPORTED_PLATFORMS
-	UE_LOG(LogSteamVR, Log, TEXT("TSteamVRTexture::FreeResource: Freeing raw texture resource. ResourceId: %d"), ResourceId);
-	VRModelManager->FreeTexture(RawResource);
+	check(VRModelManager);
+	UE_LOG(LogSteamVR, Log, TEXT("FSteamVRTextureLoader::FreeResource_Internal: Freeing texture resource %p"), RawResourcePtr);
+	VRModelManager->FreeTexture(RawResourcePtr);
 #endif
 }
 
+/* TSteamVRResource
+ *****************************************************************************/
+
+template< typename ResType, typename IDType >
+struct TSteamVRResource
+{
+private:
+	typedef TSteamVRResourceLoader<ResType, IDType> FResourceLoader;
+	static  FResourceLoader ResourceLoader;
+
+public:
+	TSteamVRResource(const IDType& ResID, bool bKickOffLoad = true)
+		: ResourceId(ResID)
+		, RawResource()
+		, bLoadFailed(false)
+	{
+		if (bKickOffLoad)
+		{
+			TickAsyncLoad();
+		}
+	}
+
+	bool IsPending()
+	{
+		return !RawResource.IsValid() && !bLoadFailed;
+	}
+
+	bool IsValid()
+	{
+		return RawResource.IsValid();
+	}
+
+	TSharedPtr<ResType> TickAsyncLoad()
+	{
+		if (IsPending())
+		{
+			bLoadFailed = !ResourceLoader.LoadResource_Async(ResourceId, RawResource);
+			if (bLoadFailed)
+			{
+				RawResource.Reset();
+			}
+		}
+		return RawResource;
+	}
+
+	operator const ResType*() const { return RawResource.Get(); }
+	const ResType* operator->() const { return RawResource.Get(); }
+
+	IDType GetId() const  { return ResourceId;  }
+
+protected:
+	IDType				ResourceId;
+	TSharedPtr<ResType> RawResource;
+	bool				bLoadFailed;
+};
+
+template<typename ResType, typename IDType>
+TSteamVRResourceLoader<ResType, IDType> TSteamVRResource<ResType, IDType>::ResourceLoader;
+
 /* FSteamVRModel 
  *****************************************************************************/
 
+typedef TSteamVRResource<vr::RenderModel_t, FString> TSteamVRModel;
 struct FSteamVRModel : public TSteamVRModel
 {
 public:
 	FSteamVRModel(const FString& ResID, bool bKickOffLoad = true)
 		: TSteamVRModel(ResID, bKickOffLoad)
 	{}
-
-	virtual ~FSteamVRModel()
-	{
-		Reset();
-	}
 public:
 	bool GetRawMeshData(float UEMeterScale, FSteamVRMeshData& MeshDataOut);
 };
@@ -282,9 +305,9 @@ struct FSteamVRMeshData
 
 bool FSteamVRModel::GetRawMeshData(float UEMeterScale, FSteamVRMeshData& MeshDataOut)
 {
-	bool bIsValidData = (RawResource != nullptr);
+	bool bIsValidData = RawResource.IsValid();
 #if STEAMVR_SUPPORTED_PLATFORMS
-	if (RawResource)
+	if (bIsValidData)
 	{
 		// Logging to try to get info about a crash where the RawMeshData seems to be bad.
 		UE_LOG(LogSteamVR, Log, TEXT("FSteamVRModel::GetRawMeshData ResourceId %s rVertexData 0x%lx unVertexCount %i rIndexData 0x%lx unTriangleCount %i diffuseTextureId %i"), *ResourceId, RawResource->rVertexData, RawResource->unVertexCount, RawResource->rIndexData, RawResource->unTriangleCount, RawResource->diffuseTextureId);
@@ -329,6 +352,7 @@ bool FSteamVRModel::GetRawMeshData(float UEMeterScale, FSteamVRMeshData& MeshDat
 /* FSteamVRTexture 
  *****************************************************************************/
 
+typedef TSteamVRResource<vr::RenderModel_TextureMap_t, int32> TSteamVRTexture;
 struct FSteamVRTexture : public TSteamVRTexture
 {
 public:
@@ -336,11 +360,6 @@ public:
 		: TSteamVRTexture(ResID, bKickOffLoad)
 	{}
 
-	~FSteamVRTexture()
-	{
-		Reset();
-	}
-
 	int32 GetResourceID() const { return ResourceId; }
 
 public:
@@ -349,7 +368,7 @@ public:
 		UTexture2D* NewTexture = nullptr;
 
 #if STEAMVR_SUPPORTED_PLATFORMS
-		if (RawResource != nullptr)
+		if (RawResource.IsValid())
 		{
 			// Create the texture. Using UTexture2D::CreateTransient which is supported outside of editor builds.
 			NewTexture = UTexture2D::CreateTransient(RawResource->unWidth, RawResource->unHeight, EPixelFormat::PF_R8G8B8A8, ObjName);
@@ -416,7 +435,7 @@ public:
 
 protected:
 	/** */
-	bool EnqueueTextureLoad(int32 SubMeshIndex, vr::RenderModel_t* RenderModel);
+	bool EnqueueTextureLoad(int32 SubMeshIndex, TSharedPtr<vr::RenderModel_t> RenderModel);
 	/** */
 	void OnLoadComplete(int32 SubMeshIndex);
 
@@ -471,7 +490,7 @@ void FSteamVRAsyncMeshLoader::Tick(float /*DeltaTime*/)
 		FSteamVRModel& ModelResource = EnqueuedMeshes[SubMeshIndex];
 		if (ModelResource.IsPending())
 		{
-			vr::RenderModel_t* RenderModel = ModelResource.TickAsyncLoad();
+			auto RenderModel = ModelResource.TickAsyncLoad();
 			if (!ModelResource.IsPending())
 			{
 				--PendingLoadCount;
@@ -552,7 +571,7 @@ void FSteamVRAsyncMeshLoader::AddReferencedObjects(FReferenceCollector& Collecto
 	Collector.AddReferencedObjects(ConstructedTextures);
 }
 
-bool FSteamVRAsyncMeshLoader::EnqueueTextureLoad(int32 SubMeshIndex, vr::RenderModel_t* RenderModel)
+bool FSteamVRAsyncMeshLoader::EnqueueTextureLoad(int32 SubMeshIndex, TSharedPtr<vr::RenderModel_t> RenderModel)
 {
 	bool bLoadEnqueued = false;
 #if STEAMVR_SUPPORTED_PLATFORMS
@@ -589,7 +608,7 @@ void FSteamVRAsyncMeshLoader::OnLoadComplete(int32 SubMeshIndex)
 				"\t %s \n" 
 				"\t Vert count: %d \n" 
 				"\t Tri  count: %d \n" 
-			"Treating as a load failure (no model will be spawned)!"), (vr::RenderModel_t*)LoadedModel, *LoadedModel.GetId(), LoadedModel->unVertexCount, LoadedModel->unTriangleCount);
+			"Treating as a load failure (no model will be spawned)!"), (const vr::RenderModel_t*)LoadedModel, *LoadedModel.GetId(), LoadedModel->unVertexCount, LoadedModel->unTriangleCount);
 
 			if (!bHasMalformData)
 			{
