PR #6094: Add Input Device Context to allow differentiating source of controller input (Contributed by [REDACTED_PERSON])

diff --git a/Engine/Plugins/Experimental/RawInput/Source/RawInput/Private/Windows/RawInputWindows.cpp b/Engine/Plugins/Experimental/RawInput/Source/RawInput/Private/Windows/RawInputWindows.cpp
--- a/Engine/Plugins/Experimental/RawInput/Source/RawInput/Private/Windows/RawInputWindows.cpp
+++ b/Engine/Plugins/Experimental/RawInput/Source/RawInput/Private/Windows/RawInputWindows.cpp
@@ -895,6 +895,8 @@ void FRawInputWindows::ShowDeviceInfo(const FConnectedDeviceInfo& DeviceInfo) co
 	}
 }
 
+static FName RawInputInterfaceName = FName("RawInput");
+
 void FRawInputWindows::SendControllerEvents()
 {
 	for (TPair<int32, FRawWindowsDeviceEntry>& DeviceEntryPair : RegisteredDeviceList)
@@ -904,6 +906,8 @@ void FRawInputWindows::SendControllerEvents()
 		// e.g if a button is still down or axis has a value (e.g. wheel not in centre)
 		if (DeviceEntry.bNeedsUpdate)
 		{
+			FInputDeviceScope InputScope(this, RawInputInterfaceName, DeviceEntryPair.Key, DeviceEntry.DeviceData.DeviceName);
+
 			for (FButtonData& DeviceButtonData : DeviceEntry.ButtonData)
 			{
 				if (!DeviceButtonData.ButtonName.IsNone())
diff --git a/Engine/Plugins/Experimental/RawInput/Source/RawInput/RawInput.Build.cs b/Engine/Plugins/Experimental/RawInput/Source/RawInput/RawInput.Build.cs
--- a/Engine/Plugins/Experimental/RawInput/Source/RawInput/RawInput.Build.cs
+++ b/Engine/Plugins/Experimental/RawInput/Source/RawInput/RawInput.Build.cs
@@ -12,15 +12,16 @@ namespace UnrealBuildTool.Rules
 				new string[]
 				{
 					"Core",
-				    "CoreUObject",
-				    "Engine",
+					"CoreUObject",
+					"Engine",
+					"ApplicationCore",
 					"InputCore",
 					"SlateCore",
 					"Slate"
 				}
 			);
 
-            PrivateDependencyModuleNames.AddRange(
+			PrivateDependencyModuleNames.AddRange(
 				new string[]
 				{
 					"InputDevice",
diff --git a/Engine/Plugins/Runtime/Steam/SteamController/Source/SteamController/Private/SteamController.cpp b/Engine/Plugins/Runtime/Steam/SteamController/Source/SteamController/Private/SteamController.cpp
--- a/Engine/Plugins/Runtime/Steam/SteamController/Source/SteamController/Private/SteamController.cpp
+++ b/Engine/Plugins/Runtime/Steam/SteamController/Source/SteamController/Private/SteamController.cpp
@@ -179,6 +179,11 @@ public:
 			ControllerHandle_t ControllerHandle = ControllerHandles[i];
 			FControllerState& ControllerState = ControllerStates[i];
 
+			// Doesn't seem to be a good way to get a non-localized string, so use generic name for scope
+			static FName SystemName(TEXT("SteamController"));
+			static FString ControllerName(TEXT("SteamController"));
+			FInputDeviceScope InputScope(this, SystemName, i, ControllerName);
+
 			for (auto It = DigitalActionHandlesMap.CreateConstIterator(); It; ++It)
 			{
 				FName DigitalActionName = It.Key();
diff --git a/Engine/Plugins/Runtime/Steam/SteamController/Source/SteamController/SteamController.Build.cs b/Engine/Plugins/Runtime/Steam/SteamController/Source/SteamController/SteamController.Build.cs
--- a/Engine/Plugins/Runtime/Steam/SteamController/Source/SteamController/SteamController.Build.cs
+++ b/Engine/Plugins/Runtime/Steam/SteamController/Source/SteamController/SteamController.Build.cs
@@ -18,6 +18,7 @@ public class SteamController : ModuleRules
         {
 			"Core",
 			"CoreUObject",
+			"ApplicationCore",
 			"Engine",
 		});
 
diff --git a/Engine/Plugins/Runtime/Steam/SteamVR/Source/SteamVRController/Private/SteamVRController.cpp b/Engine/Plugins/Runtime/Steam/SteamVR/Source/SteamVRController/Private/SteamVRController.cpp
--- a/Engine/Plugins/Runtime/Steam/SteamVR/Source/SteamVRController/Private/SteamVRController.cpp
+++ b/Engine/Plugins/Runtime/Steam/SteamVR/Source/SteamVRController/Private/SteamVRController.cpp
@@ -385,6 +385,10 @@ public:
 				{
 					if (VRControllerState.unPacketNum != ControllerState.PacketNum )
 					{
+						// Register a controller scope, using the model name
+						static FName SystemName(TEXT("SteamVR"));
+						FInputDeviceScope InputScope(this, SystemName, DeviceIndex, ControllerState.ControllerName);
+
 						bool CurrentStates[ ESteamVRControllerButton::TotalButtonCount ] = {0};
 
 						// Get the current state of all buttons
@@ -551,6 +555,8 @@ public:
 				return;
 			}
 
+			static FName SystemName(TEXT("SteamVR"));
+
 			for (auto& Action : Actions)
 			{
 				switch (Action.Type)
@@ -563,6 +569,23 @@ public:
 					{
 						if (Data.bState != Action.bState)
 						{
+							// Find the action origin
+							int32 DeviceIndex = INDEX_NONE;
+							FString ControllerName;
+							vr::InputOriginInfo_t ActionOrigin;
+							Err = VRInput->GetOriginTrackedDeviceInfo(Data.activeOrigin, &ActionOrigin, sizeof(ActionOrigin));
+
+							if (Err == vr::VRInputError_None)
+							{
+								DeviceIndex = ActionOrigin.trackedDeviceIndex;
+								if (DeviceIndex > 0 && DeviceIndex <= MaxControllers)
+								{
+									ControllerName = ControllerStates[DeviceIndex].ControllerName;
+								}
+							}
+
+							FInputDeviceScope InputScope(this, SystemName, DeviceIndex, ControllerName);
+
 							Action.bState = Data.bState;
 							if (Action.bState)
 							{
@@ -591,6 +614,23 @@ public:
 					Err = VRInput->GetAnalogActionData(Action.Handle, &Data, sizeof(Data), vr::k_ulInvalidInputValueHandle);
 					if (Err == vr::VRInputError_None)
 					{
+						// Find the action origin
+						int32 DeviceIndex = INDEX_NONE;
+						FString ControllerName;
+						vr::InputOriginInfo_t ActionOrigin;
+						Err = VRInput->GetOriginTrackedDeviceInfo(Data.activeOrigin, &ActionOrigin, sizeof(ActionOrigin));
+
+						if (Err == vr::VRInputError_None)
+						{
+							DeviceIndex = ActionOrigin.trackedDeviceIndex;
+							if (DeviceIndex > 0 && DeviceIndex <= MaxControllers)
+							{
+								ControllerName = ControllerStates[DeviceIndex].ControllerName;
+							}
+						}
+
+						FInputDeviceScope InputScope(this, SystemName, DeviceIndex, ControllerName);
+
 						if (!Action.ActionKey_X.IsNone() && Data.x != Action.Value.X)
 						{
 							Action.Value.X = Data.x;
@@ -963,7 +1003,24 @@ private:
 		// determine which player controller to assign the device to
 		int32 ControllerIndex = FMath::FloorToInt(NumControllersMapped / CONTROLLERS_PER_PLAYER);
 
-		UE_LOG(LogSteamVRController, Verbose, TEXT("Controller device %i is being assigned unreal hand %i (left=0, right=1), for player %i."), DeviceIndex, (int32)ChosenHand, ControllerIndex);
+		// read the model name for later use
+		vr::TrackedPropertyError APIError;
+		TArray<char> Buffer;
+		Buffer.AddUninitialized(vr::k_unMaxPropertyStringSize);
+
+		int Size = VRSystem->GetStringTrackedDeviceProperty(DeviceIndex, vr::Prop_RenderModelName_String, Buffer.GetData(), Buffer.Num(), &APIError);
+		if (APIError == vr::TrackedProp_BufferTooSmall)
+		{
+			Buffer.AddUninitialized(Size - Buffer.Num());
+			Size = VRSystem->GetStringTrackedDeviceProperty(DeviceIndex, vr::Prop_RenderModelName_String, Buffer.GetData(), Buffer.Num(), &APIError);
+		}
+
+		if (APIError == vr::TrackedProp_Success)
+		{
+			ControllerStates[DeviceIndex].ControllerName = UTF8_TO_TCHAR(Buffer.GetData());
+		}
+
+		UE_LOG(LogSteamVRController, Verbose, TEXT("Controller device %i (%s) is being assigned unreal hand %i (left=0, right=1), for player %i."), *ControllerStates[DeviceIndex].ControllerName, DeviceIndex, (int32)ChosenHand, ControllerIndex);
 		ControllerStates[DeviceIndex].Hand = ChosenHand;
 		UnrealControllerHandUsageCount[(int32)ChosenHand] += 1;
 
@@ -1559,6 +1616,9 @@ private:
 		/** Which hand this controller is representing */
 		EControllerHand Hand;
 
+		/** Name of this controller, uses the render model name as that is most useful */
+		FString ControllerName;
+
 		/** If packet num matches that on your prior call, then the controller state hasn't been changed since 
 		  * your last call and there is no need to process it. */
 		uint32 PacketNum;
diff --git a/Engine/Source/Runtime/ApplicationCore/Private/GenericPlatform/GenericApplication.cpp b/Engine/Source/Runtime/ApplicationCore/Private/GenericPlatform/GenericApplication.cpp
--- a/Engine/Source/Runtime/ApplicationCore/Private/GenericPlatform/GenericApplication.cpp
+++ b/Engine/Source/Runtime/ApplicationCore/Private/GenericPlatform/GenericApplication.cpp
@@ -3,6 +3,7 @@
 #include "GenericPlatform/GenericApplication.h"
 #include "HAL/IConsoleManager.h"
 #include "Misc/CoreDelegates.h"
+#include "IInputDevice.h"
 
 const FGamepadKeyNames::Type FGamepadKeyNames::Invalid(NAME_None);
 
@@ -101,6 +102,43 @@ const FGamepadKeyNames::Type FGamepadKeyNames::MotionController_Right_TriggerAxi
 const FGamepadKeyNames::Type FGamepadKeyNames::MotionController_Right_Grip1Axis( "MotionController_Right_Grip1Axis" );
 const FGamepadKeyNames::Type FGamepadKeyNames::MotionController_Right_Grip2Axis( "MotionController_Right_Grip2Axis" );
 
+TArray<FInputDeviceScope*> FInputDeviceScope::ScopeStack;
+
+FInputDeviceScope::FInputDeviceScope(IInputDevice* InInputDevice, FName InInputDeviceName, int32 InHardwareDeviceHandle, FString InHardwareDeviceIdentifier)
+	: InputDevice(InInputDevice)
+	, InputDeviceName(InInputDeviceName)
+	, HardwareDeviceHandle(InHardwareDeviceHandle)
+	, HardwareDeviceIdentifier(InHardwareDeviceIdentifier)
+{
+	if (ensure(IsInGameThread()))
+	{
+		// Add to scope stack
+		ScopeStack.Add(this);
+	}
+}
+
+FInputDeviceScope::~FInputDeviceScope()
+{
+	if (ensure(IsInGameThread()))
+	{
+		// This should always be the top of the stack
+		ensureMsgf((ScopeStack.Num() > 0 && ScopeStack.Last() == this), TEXT("FInputDeviceScope was not destroyed in correct order!"));
+		ScopeStack.Remove(this);
+	}
+}
+
+const FInputDeviceScope* FInputDeviceScope::GetCurrent()
+{
+	if (ensure(IsInGameThread()))
+	{
+		if (ScopeStack.Num() > 0)
+		{
+			return ScopeStack.Last();
+		}
+	}
+	return nullptr;
+}
+
 float GDebugSafeZoneRatio = 1.0f;
 float GDebugActionZoneRatio = 1.0f;
 
diff --git a/Engine/Source/Runtime/ApplicationCore/Private/Windows/XInputInterface.cpp b/Engine/Source/Runtime/ApplicationCore/Private/Windows/XInputInterface.cpp
--- a/Engine/Source/Runtime/ApplicationCore/Private/Windows/XInputInterface.cpp
+++ b/Engine/Source/Runtime/ApplicationCore/Private/Windows/XInputInterface.cpp
@@ -98,6 +98,9 @@ float ShortToNormalizedFloat(int16 AxisVal)
 	return float(AxisVal) / Norm;
 }
 
+static FName XInputInterfaceName = FName("XInputInterface");
+static FString XInputControllerIdentifier = TEXT("XInputController");
+
 void XInputInterface::SendControllerEvents()
 {
 	bool bWereConnected[MAX_NUM_XINPUT_CONTROLLERS];
@@ -126,6 +129,9 @@ void XInputInterface::SendControllerEvents()
 		
 	for ( int32 ControllerIndex = 0; ControllerIndex < MAX_NUM_XINPUT_CONTROLLERS; ++ControllerIndex )
 	{
+		// Set input scope, there doesn't seem to be a reliable way to differentiate 360 vs Xbox one controllers so use generic name
+		FInputDeviceScope InputScope(this, XInputInterfaceName, ControllerIndex, XInputControllerIdentifier);
+
 		FControllerState& ControllerState = ControllerStates[ControllerIndex];
 
 		const bool bWasConnected = bWereConnected[ControllerIndex];
@@ -256,7 +262,7 @@ void XInputInterface::SetMessageHandler( const TSharedRef< FGenericApplicationMe
 	MessageHandler = InMessageHandler;
 }
 
-void XInputInterface::SetChannelValue( const int32 ControllerId, const FForceFeedbackChannelType ChannelType, const float Value )
+void XInputInterface::SetChannelValue( int32 ControllerId, const FForceFeedbackChannelType ChannelType, const float Value )
 {
 	if (ControllerId >= 0 && ControllerId < MAX_NUM_XINPUT_CONTROLLERS)
 	{
@@ -286,7 +292,7 @@ void XInputInterface::SetChannelValue( const int32 ControllerId, const FForceFee
 	}
 }
 
-void XInputInterface::SetChannelValues( const int32 ControllerId, const FForceFeedbackValues &Values )
+void XInputInterface::SetChannelValues( int32 ControllerId, const FForceFeedbackValues &Values )
 {
 	if (ControllerId >= 0 && ControllerId < MAX_NUM_XINPUT_CONTROLLERS)
 	{
diff --git a/Engine/Source/Runtime/ApplicationCore/Private/Windows/XInputInterface.h b/Engine/Source/Runtime/ApplicationCore/Private/Windows/XInputInterface.h
--- a/Engine/Source/Runtime/ApplicationCore/Private/Windows/XInputInterface.h
+++ b/Engine/Source/Runtime/ApplicationCore/Private/Windows/XInputInterface.h
@@ -4,6 +4,7 @@
 
 #include "CoreTypes.h"
 #include "GenericPlatform/IInputInterface.h"
+#include "IInputDevice.h"
 #include "GenericPlatform/GenericApplicationMessageHandler.h"
 
 /** Max number of controllers. */
@@ -17,7 +18,7 @@ enum class FForceFeedbackChannelType;
 /**
  * Interface class for XInput devices (xbox 360 controller)                 
  */
-class XInputInterface
+class XInputInterface : public IInputDevice
 {
 public:
 
@@ -28,9 +29,9 @@ public:
 	 *
 	 * @param PathToJoystickCaptureWidget	The path to the joystick capture widget.  If invalid this function does not poll 
 	 */
-	void SendControllerEvents();
+	virtual void SendControllerEvents() override;
 
-	void SetMessageHandler( const TSharedRef< FGenericApplicationMessageHandler >& InMessageHandler );
+	virtual void SetMessageHandler( const TSharedRef< FGenericApplicationMessageHandler >& InMessageHandler ) override;
 
 	void SetNeedsControllerStateUpdate() { bNeedsControllerStateUpdate = true; }
 
@@ -42,7 +43,7 @@ public:
 	* @param ChannelType the type of channel whose value should be set
 	* @param Value strength or speed of feedback, 0.0f to 1.0f. 0.0f will disable
 	*/
-	void SetChannelValue( const int32 ControllerId, const FForceFeedbackChannelType ChannelType, const float Value );
+	virtual void SetChannelValue( int32 ControllerId, const FForceFeedbackChannelType ChannelType, const float Value ) override;
 
 	/**
 	* Sets the strength/speed of all the channels for the given controller id.
@@ -51,9 +52,11 @@ public:
 	* @param ControllerId the id of the controller whose value is to be set
 	* @param Values strength or speed of feedback for all channels
 	*/
-	void SetChannelValues( const int32 ControllerId, const FForceFeedbackValues& Values );
+	virtual void SetChannelValues( int32 ControllerId, const FForceFeedbackValues& Values ) override;
 
-	bool IsGamepadAttached() const { return bIsGamepadAttached; }
+	virtual bool IsGamepadAttached() const override { return bIsGamepadAttached; }
+	virtual void Tick( float DeltaTime ) override {};
+	virtual bool Exec( UWorld* InWorld, const TCHAR* Cmd, FOutputDevice& Ar ) { return false; }
 
 private:
 
diff --git a/Engine/Source/Runtime/ApplicationCore/Public/GenericPlatform/GenericApplicationMessageHandler.h b/Engine/Source/Runtime/ApplicationCore/Public/GenericPlatform/GenericApplicationMessageHandler.h
--- a/Engine/Source/Runtime/ApplicationCore/Public/GenericPlatform/GenericApplicationMessageHandler.h
+++ b/Engine/Source/Runtime/ApplicationCore/Public/GenericPlatform/GenericApplicationMessageHandler.h
@@ -227,7 +227,45 @@ private:
 	TOptional<float> MaxHeight;
 };
 
+/** 
+ * Context scope that indicates which IInputDevice is currently being handled. 
+ * This can be used to determine hardware-specific information when handling input from FGenericApplicationMessageHandler subclasses.
+ * This is generally set during SendControllerEvents or Tick and is only valid on the game thread.
+ */
+class APPLICATIONCORE_API FInputDeviceScope
+{
+public:
+	/** The specific InputDevice that is currently being polled. This is only valid within the current function scope and may be null */
+	class IInputDevice* InputDevice;
+
+	/** Logical name of the input device interface. This is not translated but is platform-specific */
+	FName InputDeviceName;
+
+	/** A system-specific device id, this is not the same as controllerId and represents a physical device instead of logical user. -1 represents an unknown device */
+	int32 HardwareDeviceHandle;
+
+	/** Logical string identifying the hardware device. This is not translated and is system-specific, it may be empty */
+	FString HardwareDeviceIdentifier;
+
+	/** Constructor, this should only be allocated directly on the stack */
+	FInputDeviceScope(IInputDevice* InInputDevice, FName InInputDeviceName, int32 InHardwareDeviceHandle = -1, FString InHardwareDeviceIdentifier = FString());
+	~FInputDeviceScope();
+
+	/** Cannot be copied/moved */
+	FInputDeviceScope() = delete;
+	FInputDeviceScope(const FInputDeviceScope&) = delete;
+	FInputDeviceScope& operator=(const FInputDeviceScope&) = delete;
+	FInputDeviceScope(FInputDeviceScope&&) = delete;
+	FInputDeviceScope& operator=(FInputDeviceScope&&) = delete;
+
+	/** Returns the currently active InputDeviceScope. This is only valid to call on the game thread and may return null */
+	static const FInputDeviceScope* GetCurrent();
+
+private:
+	static TArray<FInputDeviceScope*> ScopeStack;
+};
 
+/** Interface that defines how to handle interaction with a user via hardware input and output */
 class FGenericApplicationMessageHandler
 {
 public:
