Fix build issue with struct initializers and unions

diff --git a/Engine/Plugins/Runtime/Steam/SteamVR/Source/SteamVRController/Private/SteamVRController.cpp b/Engine/Plugins/Runtime/Steam/SteamVR/Source/SteamVRController/Private/SteamVRController.cpp
--- a/Engine/Plugins/Runtime/Steam/SteamVR/Source/SteamVRController/Private/SteamVRController.cpp
+++ b/Engine/Plugins/Runtime/Steam/SteamVR/Source/SteamVRController/Private/SteamVRController.cpp
@@ -1100,7 +1100,7 @@ private:
 					if (KeyMapping != nullptr)
 					{
 						FString ActionPath = FString("/actions/main/in") / ActionName.ToString();
-						Actions.Add(FSteamVRAction{ ActionPath, FSteamVRAction::Boolean, ActionName, KeyMapping->Key.GetFName(), FName(), false });
+						Actions.Add(FSteamVRAction( ActionPath, ActionName, KeyMapping->Key.GetFName(), false ));
 					}
 				}
 
@@ -1120,7 +1120,7 @@ private:
 					if (KeyMapping != nullptr)
 					{
 						FString ActionPath = FString("/actions/main/in") / AxisName.ToString() + TEXT("_axis");
-						Actions.Add(FSteamVRAction{ ActionPath, FSteamVRAction::Vector1, AxisName, KeyMapping->Key.GetFName(), FName(), 0.0f });
+						Actions.Add(FSteamVRAction( ActionPath, AxisName, KeyMapping->Key.GetFName(), 0.0f ));
 					}
 				}
 
@@ -1129,7 +1129,7 @@ private:
 					const FKey* ConsoleKey = InputSettings->ConsoleKeys.FindByPredicate([](FKey& Key) { return Key.IsValid(); });
 					if (ConsoleKey != nullptr)
 					{
-						Actions.Add(FSteamVRAction{ FString("/actions/main/in/open_console"), FSteamVRAction::Boolean, FName(TEXT("Open Console")), ConsoleKey->GetFName(), FName(), false });
+						Actions.Add(FSteamVRAction( FString("/actions/main/in/open_console"), FName(TEXT("Open Console")), ConsoleKey->GetFName(), false ));
 					}
 				}
 			}
@@ -1307,6 +1307,38 @@ private:
 
 			return TypeStrings[(int)Type];
 		}
+
+		FSteamVRAction(const FString& inPath, const FName& inName, const FName& inActionKey, bool inState)
+			: Path(inPath)
+			, Type(Boolean)
+			, Name(inName)
+			, ActionKey(inActionKey)
+			, ActionKey_Y()
+			, Value2D()
+		{
+			bState = inState;
+		}
+
+		FSteamVRAction(const FString& inPath, const FName& inName, const FName& inActionKey, float inValue1D)
+			: Path(inPath)
+			, Type(Vector1)
+			, Name(inName)
+			, ActionKey(inActionKey)
+			, ActionKey_Y()
+			, Value2D()
+		{
+			Value1D = inValue1D;
+		}
+
+		FSteamVRAction(const FString& inPath, const FName& inName, const FName& inActionKey_X, const FName& inActionKey_Y, const FVector2D& inValue2D)
+			: Path(inPath)
+			, Type(Vector2)
+			, Name(inName)
+			, ActionKey(inActionKey_X)
+			, ActionKey_Y(inActionKey_Y)
+			, Value2D(inValue2D)
+		{}
+
 	};
 
 	/** Mappings between tracked devices and 0 indexed controllers */
