PR #5019: [BUGFIX SteamVR Vulkan] Add missing method to update SteamVR Vulkan viewport. (Contributed by [REDACTED_PERSON])

diff --git a/Engine/Plugins/Runtime/Steam/SteamVR/Source/SteamVR/Private/SteamVRHMD.h b/Engine/Plugins/Runtime/Steam/SteamVR/Source/SteamVR/Private/SteamVRHMD.h
--- a/Engine/Plugins/Runtime/Steam/SteamVR/Source/SteamVR/Private/SteamVRHMD.h
+++ b/Engine/Plugins/Runtime/Steam/SteamVR/Source/SteamVR/Private/SteamVRHMD.h
@@ -309,6 +309,7 @@ public:
 		VulkanBridge(FSteamVRHMD* plugin);
 
 		virtual void FinishRendering() override;
+		virtual void UpdateViewport(const FViewport& Viewport, FRHIViewport* InViewportRHI) override;
 		virtual void Reset() override;
 
 	protected:
diff --git a/Engine/Plugins/Runtime/Steam/SteamVR/Source/SteamVR/Private/SteamVRRender.cpp b/Engine/Plugins/Runtime/Steam/SteamVR/Source/SteamVR/Private/SteamVRRender.cpp
--- a/Engine/Plugins/Runtime/Steam/SteamVR/Source/SteamVR/Private/SteamVRRender.cpp
+++ b/Engine/Plugins/Runtime/Steam/SteamVR/Source/SteamVR/Private/SteamVRRender.cpp
@@ -281,6 +281,12 @@ void FSteamVRHMD::VulkanBridge::Reset()
 
 }
 
+void FSteamVRHMD::VulkanBridge::UpdateViewport(const FViewport& Viewport, FRHIViewport* InViewportRHI)
+{
+	RenderTargetTexture = Viewport.GetRenderTargetTexture();
+	check(IsValidRef(RenderTargetTexture));
+}
+
 FSteamVRHMD::OpenGLBridge::OpenGLBridge(FSteamVRHMD* plugin):
 	BridgeBaseImpl(plugin),
 	RenderTargetTexture(0)
