Map X/Y axis pairs to a single vector2 input in Steam VR Input
(And also support for vector3 inputs for X/Y/Z key triplets, although there are currently no key names that fit the restriction.)

diff --git a/Engine/Plugins/Runtime/Steam/SteamVR/Source/SteamVRController/Private/SteamVRController.cpp b/Engine/Plugins/Runtime/Steam/SteamVR/Source/SteamVRController/Private/SteamVRController.cpp
--- a/Engine/Plugins/Runtime/Steam/SteamVR/Source/SteamVRController/Private/SteamVRController.cpp
+++ b/Engine/Plugins/Runtime/Steam/SteamVR/Source/SteamVRController/Private/SteamVRController.cpp
@@ -550,43 +550,60 @@ public:
 				{
 					vr::InputDigitalActionData_t Data;
 					Err = VRInput->GetDigitalActionData(Action.Handle, &Data, sizeof(Data), vr::k_ulInvalidInputValueHandle);
-					if (Err != vr::VRInputError_None)
-					{
-						UE_LOG(LogSteamVRController, Warning, TEXT("GetDigitalActionData for %s returned error: %d"), *Action.Name.ToString(), (int32)Err);
-					}
-					else
+					if (Err == vr::VRInputError_None)
 					{
 						if (Data.bState != Action.bState)
 						{
 							Action.bState = Data.bState;
 							if (Action.bState)
 							{
-								MessageHandler->OnControllerButtonPressed(Action.ActionKey, 0, /*IsRepeat =*/false);
+								MessageHandler->OnControllerButtonPressed(Action.ActionKey_X, 0, /*IsRepeat =*/false);
 							}
 							else
 							{
-								MessageHandler->OnControllerButtonReleased(Action.ActionKey, 0, /*IsRepeat =*/false);
+								MessageHandler->OnControllerButtonReleased(Action.ActionKey_X, 0, /*IsRepeat =*/false);
 							}
 						}
 					}
+					// If the current error is the same as the last frame's error, don't log it again to avoid spamming the log
+					else if (Err != Action.LastError)
+					{
+						UE_LOG(LogSteamVRController, Warning, TEXT("GetDigitalActionData for %s returned error: %d"), *Action.Name.ToString(), (int32)Err);
+
+					}
+					Action.LastError = Err;
 				}
 				break;
 				case FSteamVRAction::Vector1:
+				case FSteamVRAction::Vector2:
+				case FSteamVRAction::Vector3:
 				{
 					vr::InputAnalogActionData_t Data;
 					Err = VRInput->GetAnalogActionData(Action.Handle, &Data, sizeof(Data), vr::k_ulInvalidInputValueHandle);
-					if (Err != vr::VRInputError_None)
+					if (Err == vr::VRInputError_None)
 					{
-						UE_LOG(LogSteamVRController, Warning, TEXT("GetAnalogActionData for %s returned error: %d"), *Action.Name.ToString(), (int32)Err);
-					}
-					else
-					{
-						if (Data.x != Action.Value1D)
+						if (!Action.ActionKey_X.IsNone() && Data.x != Action.Value.X)
+						{
+							Action.Value.X = Data.x;
+							MessageHandler->OnControllerAnalog(Action.ActionKey_X, 0, Action.Value.X);
+						}
+						if (!Action.ActionKey_Y.IsNone() && Data.y != Action.Value.Y)
+						{
+							Action.Value.Y = Data.y;
+							MessageHandler->OnControllerAnalog(Action.ActionKey_Y, 0, Action.Value.Y);
+						}
+						if (!Action.ActionKey_Z.IsNone() && Data.z != Action.Value.Z)
 						{
-							Action.Value1D = Data.x;
-							MessageHandler->OnControllerAnalog(Action.ActionKey, 0, Action.Value1D);
+							Action.Value.Z = Data.z;
+							MessageHandler->OnControllerAnalog(Action.ActionKey_Z, 0, Action.Value.Z);
 						}
 					}
+					// If the current error is the same as the last frame's error, don't log it again to avoid spamming the log
+					else if (Err != Action.LastError)
+					{
+						UE_LOG(LogSteamVRController, Warning, TEXT("GetAnalogActionData for %s returned error: %d"), *Action.Name.ToString(), (int32)Err);
+					}
+					Action.LastError = Err;
 				}
 				break;
 				default:
@@ -1069,11 +1086,141 @@ private:
 	}
 
 	// Hack to prefer emitting MotionController keys for action events
-	static bool MatchMotionControllerKey(const FKey& Key)
+	static bool MatchKeyNamePrefix(const FKey& Key, const TCHAR* Prefix)
 	{
-		return Key.GetFName().ToString().StartsWith(TEXT("MotionController"));
+		return Key.GetFName().ToString().StartsWith(Prefix);
 	};
 
+	static bool MatchKeyNameSuffix(const FKey& Key, const TCHAR* Suffix)
+	{
+		return Key.GetFName().ToString().EndsWith(Suffix);
+	};
+
+	// Finds an axis key mapping from a list of mapping with the following preferences:
+	// 1. Tries to find a FloatAxis key that starts with "MotionController" and ends with "X"
+	// 2. Tries to find a FloatAxis key that starts with "MotionController" and ends with "Y"
+	// 3. Any to find a FloatAxis key that starts with "MotionController"
+	// 4. Any FloatAxis that ends with "X"
+	// 5. Any FloatAxis that ends with "Y"
+	// 6. Any FloatAxis
+	// 7. Any Valid key.
+	// If case 1 or 3 is matched bOutIsXAxis will be set to true. 
+	static FName FindAxisKeyMapping(TArray<FInputAxisKeyMapping>& Mappings, bool& bOutIsXAxis)
+	{
+		FInputAxisKeyMapping* Found = nullptr;
+		bOutIsXAxis = false;
+		// First filter out all floatAxes, as all except the 5th case require a float axis.
+		TArray<FInputAxisKeyMapping> FloatMappings = Mappings.FilterByPredicate([](const FInputAxisKeyMapping& Mapping)
+		{
+			return Mapping.Key.IsFloatAxis();
+		});
+
+		// If there were no float axis key bindings, return the first valid mapping
+		if (FloatMappings.Num() == 0)
+		{
+			Found = Mappings.FindByPredicate([](FInputAxisKeyMapping& Mapping)
+			{
+				return Mapping.Key.IsValid();
+			});
+
+			if (Found != nullptr)
+			{
+				return Found->Key.GetFName();
+			}
+			else
+			{
+				return FName();
+			}
+		}
+
+		// Then get all mappings with keys starting with "MotionController"
+		TArray<FInputAxisKeyMapping> MotionControllerMappings = FloatMappings.FilterByPredicate([](const FInputAxisKeyMapping& Mapping)
+		{
+			return MatchKeyNamePrefix(Mapping.Key, TEXT("MotionController"));
+		});
+
+		// If there are no MotionController keys, search through all FloatAxes:
+		TArray<FInputAxisKeyMapping>& MappingsSubset = MotionControllerMappings.Num() == 0 ? FloatMappings : MotionControllerMappings;
+
+		Found = FloatMappings.FindByPredicate([](const FInputAxisKeyMapping& Mapping)
+		{
+			return MatchKeyNameSuffix(Mapping.Key, TEXT("X"));
+		});
+		if (Found != nullptr)
+		{
+			bOutIsXAxis = true;
+			return Found->Key.GetFName();
+		}
+
+		Found = FloatMappings.FindByPredicate([](const FInputAxisKeyMapping& Mapping)
+		{
+			return MatchKeyNameSuffix(Mapping.Key, TEXT("Y"));
+		});
+		if (Found != nullptr)
+		{
+			return Found->Key.GetFName();
+		}
+
+
+		Found = FloatMappings.FindByPredicate([](const FInputAxisKeyMapping& Mapping)
+		{
+			return Mapping.Key.IsValid();
+		});
+		if (Found != nullptr)
+		{
+			return Found->Key.GetFName();
+		}
+		else
+		{
+			return FName();
+		}
+	}
+
+	/** Returns the concatenation of two strings, skipping all characters at the beginning of string B that match the beginning of string A and
+	    all characters at the end of string A that match the end of string B.
+		Example: passing in "MoveUpAction" and "MoveRightAction" should result in "MoveUpRightAction"
+		If the strings have no common suffix or prefix, the result will simply be the concatenation of both strings.
+		If the strings are identical, returns the first string.
+
+		The algorithm treats separator characters ' ', '_' and '/' differently. If either the suffix begins with one or the prefix ends with one,
+		the function will keep one of them in the resulting string.
+		Example "move_up_action" and "move_right_action" will result in "move_up_right_action" and not "move_upright_action"
+	   */
+	static FString MergeActionNames(const FString& A, const FString& B)
+	{
+		if (A.Equals(B, ESearchCase::CaseSensitive))
+		{
+			return A;
+		}
+		const int LastA = A.Len() - 1;
+		const int LastB = B.Len() - 1;
+		const int MinLen = (LastA < LastB) ? A.Len() : B.Len();
+
+		int CommonPrefix = 0;
+		int CommonSuffix = 0;
+		for (;CommonPrefix < MinLen && A[CommonPrefix] == B[CommonPrefix]; CommonPrefix++)
+		{
+			/* intentionally blank */
+		}
+
+		for (;CommonSuffix < MinLen && A[LastA - CommonSuffix] == B[LastB - CommonSuffix]; CommonSuffix++)
+		{
+			/* intentionally blank */
+		}
+
+		// If either the common prefix ends with or the common suffix begins with a space, an underscore or a dash, keep one of them.
+		if (CommonPrefix > 0 && (A[CommonPrefix - 1] == TEXT(' ') || A[CommonPrefix - 1] == TEXT('_') || A[CommonPrefix - 1] == TEXT('-')))
+		{
+			CommonPrefix--;
+		}
+		else if (CommonSuffix > 0 && (A[LastA - CommonSuffix + 1] == TEXT(' ') || A[LastA - CommonSuffix + 1] == TEXT('_') || A[LastA - CommonSuffix + 1] == TEXT('-')))
+		{
+			CommonSuffix--;
+		}
+
+		return A.LeftChop(CommonSuffix) + B.RightChop(CommonPrefix);
+	}
+
 	void BuildActionManifest()
 	{
 		vr::IVRInput* VRInput;
@@ -1091,7 +1238,11 @@ private:
 					TArray<FInputActionKeyMapping> Mappings;
 					InputSettings->GetActionMappingByName(ActionName, Mappings);
 
-					FInputActionKeyMapping* KeyMapping = Mappings.FindByPredicate([](FInputActionKeyMapping& Mapping) { return MatchMotionControllerKey(Mapping.Key); });
+					FInputActionKeyMapping* KeyMapping = Mappings.FindByPredicate([](FInputActionKeyMapping& Mapping) 
+					{ 
+						return MatchKeyNamePrefix(Mapping.Key, TEXT("MotionController")); 
+					});
+
 					if (KeyMapping == nullptr)
 					{
 						KeyMapping = Mappings.FindByPredicate([](FInputActionKeyMapping& Mapping) { return Mapping.Key.IsValid(); });
@@ -1108,19 +1259,74 @@ private:
 				InputSettings->GetAxisNames(AxisNames);
 				for (const auto& AxisName : AxisNames)
 				{
+					bool bIsXAxis = false;
 					TArray<FInputAxisKeyMapping> Mappings;
 					InputSettings->GetAxisMappingByName(AxisName, Mappings);
 
-					FInputAxisKeyMapping* KeyMapping = Mappings.FindByPredicate([](FInputAxisKeyMapping& Mapping) { return MatchMotionControllerKey(Mapping.Key); });
-					if (KeyMapping == nullptr)
-					{
-						KeyMapping = Mappings.FindByPredicate([](FInputAxisKeyMapping& Mapping) { return Mapping.Key.IsValid(); });
-					}
+					FName KeyName = FindAxisKeyMapping(Mappings, bIsXAxis);
 
-					if (KeyMapping != nullptr)
+					if (!KeyName.IsNone())
 					{
 						FString ActionPath = FString("/actions/main/in") / AxisName.ToString() + TEXT("_axis");
-						Actions.Add(FSteamVRAction( ActionPath, AxisName, KeyMapping->Key.GetFName(), 0.0f ));
+						Actions.Add(FSteamVRAction( ActionPath, AxisName, KeyName, 0.0f ));
+
+						// If the current axis is bound to an X axis, find the corresponding Y axis binding and create
+						// a combined vector2 action from them (and if there were Z axes, create vector3 actions.)
+						if (bIsXAxis)
+						{
+							FName YKeyName = FName(*(KeyName.ToString().LeftChop(1) + TEXT('Y')));
+							FName ZKeyName = FName(*(KeyName.ToString().LeftChop(1) + TEXT('Z')));
+							FName YAxisName, ZAxisName;
+
+							for (const auto& InnerAxisName : AxisNames)
+							{
+								TArray<FInputAxisKeyMapping> InnerMappings;
+								InputSettings->GetAxisMappingByName(InnerAxisName, InnerMappings);
+
+								if (YAxisName.IsNone() && 
+									InnerMappings.ContainsByPredicate([&YKeyName](auto& Mapping)
+									{
+										return Mapping.Key.GetFName() == YKeyName;
+									})
+								)
+								{
+									YAxisName = InnerAxisName;
+								}
+
+								if (ZAxisName.IsNone() && 
+									InnerMappings.ContainsByPredicate([&ZKeyName](auto& Mapping)
+									{
+										return Mapping.Key.GetFName() == ZKeyName;
+									})
+								)
+								{
+									ZAxisName = InnerAxisName;
+								}
+
+								if (!YAxisName.IsNone() && !ZAxisName.IsNone())
+								{
+									break;
+								}
+
+							}
+
+							if (!YAxisName.IsNone())
+							{
+								FString CombinedAxisName = MergeActionNames(AxisName.ToString(), YAxisName.ToString());
+								if (!ZAxisName.IsNone())
+								{
+									CombinedAxisName = MergeActionNames(CombinedAxisName, ZAxisName.ToString());
+									FString CombinedActionPath = FString("/actions/main/in") / CombinedAxisName + TEXT("_axis3d");
+									Actions.Add(FSteamVRAction(CombinedActionPath, FName(*CombinedAxisName), KeyName, YKeyName, ZKeyName, FVector::ZeroVector));
+								}
+								else
+								{
+									FString CombinedActionPath = FString("/actions/main/in") / CombinedAxisName + TEXT("_axis2d");
+									Actions.Add(FSteamVRAction(CombinedActionPath, FName(*CombinedAxisName), KeyName, YKeyName, FVector2D::ZeroVector));
+								}
+							}
+									
+						}
 					}
 				}
 
@@ -1282,15 +1488,16 @@ private:
 		FString		Path;
 		EActionType	Type;
 		FName		Name;
-		FName		ActionKey;
-		FName		ActionKey_Y; // TODO: Emulate vector2 actions by combining input axes mapped to an X/Y key combination.
+		FName		ActionKey_X;
+		FName		ActionKey_Y;
+		FName		ActionKey_Z;
 		union {
 			bool	bState;
-			float	Value1D;
-			FVector2D Value2D;
+			FVector Value;
 		};
 
 		vr::VRActionHandle_t Handle;
+		vr::EVRInputError LastError;
 
 		FString TypeAsString()
 		{
@@ -1312,9 +1519,12 @@ private:
 			: Path(inPath)
 			, Type(Boolean)
 			, Name(inName)
-			, ActionKey(inActionKey)
+			, ActionKey_X(inActionKey)
 			, ActionKey_Y()
-			, Value2D()
+			, ActionKey_Z()
+			, Value()
+			, Handle()
+			, LastError(vr::VRInputError_None)
 		{
 			bState = inState;
 		}
@@ -1323,20 +1533,36 @@ private:
 			: Path(inPath)
 			, Type(Vector1)
 			, Name(inName)
-			, ActionKey(inActionKey)
+			, ActionKey_X(inActionKey)
 			, ActionKey_Y()
-			, Value2D()
-		{
-			Value1D = inValue1D;
-		}
+			, ActionKey_Z()
+			, Value(inValue1D, 0, 0)
+			, Handle()
+			, LastError(vr::VRInputError_None)
+		{}
 
 		FSteamVRAction(const FString& inPath, const FName& inName, const FName& inActionKey_X, const FName& inActionKey_Y, const FVector2D& inValue2D)
 			: Path(inPath)
 			, Type(Vector2)
 			, Name(inName)
-			, ActionKey(inActionKey_X)
+			, ActionKey_X(inActionKey_X)
+			, ActionKey_Y(inActionKey_Y)
+			, ActionKey_Z()
+			, Value(inValue2D.X, inValue2D.Y, 0)
+			, Handle()
+			, LastError(vr::VRInputError_None)
+		{}
+
+		FSteamVRAction(const FString& inPath, const FName& inName, const FName& inActionKey_X, const FName& inActionKey_Y, const FName& inActionKey_Z, const FVector& inValue3D)
+			: Path(inPath)
+			, Type(Vector2)
+			, Name(inName)
+			, ActionKey_X(inActionKey_X)
 			, ActionKey_Y(inActionKey_Y)
-			, Value2D(inValue2D)
+			, ActionKey_Z(inActionKey_Z)
+			, Value(inValue3D)
+			, Handle()
+			, LastError(vr::VRInputError_None)
 		{}
 
 	};
