Ensure default bindings exported from the SteamVR bindings editor are referenced by the generated Action Manifest JSON file. #jira UE-64098
Store the generated action manifest file in the correct directory on Mac.

diff --git a/Engine/Plugins/Runtime/Steam/SteamVR/Source/SteamVRController/Private/SteamVRController.cpp b/Engine/Plugins/Runtime/Steam/SteamVR/Source/SteamVRController/Private/SteamVRController.cpp
--- a/Engine/Plugins/Runtime/Steam/SteamVR/Source/SteamVRController/Private/SteamVRController.cpp
+++ b/Engine/Plugins/Runtime/Steam/SteamVR/Source/SteamVRController/Private/SteamVRController.cpp
@@ -11,6 +11,7 @@
 #include "GameFramework/InputSettings.h"
 #include "../../SteamVR/Private/SteamVRHMD.h"
 #include "SteamVRControllerLibrary.h" // for ESteamVRTouchDPadMapping
+#include "HAL/FileManagerGeneric.h"
 
 #if WITH_EDITOR
 #include "Editor.h"
@@ -1341,7 +1342,10 @@ private:
 
 			if (Actions.Num() > 0)
 			{
-				FString ManifestPath = FString(FPaths::ProjectConfigDir() / TEXT("steamvr_actions.json"));
+				// The steamvr_actions.json file is generated from internal data, so it goes to GeneratedConfig directory.
+				const FString ManifestPath = FPaths::GeneratedConfigDir() / TEXT("steamvr_actions.json");
+				// The default bindings files need to be generated by the developer, so they are stored in the project config directory.
+				const FString BindingsDir = FPaths::ProjectConfigDir() / TEXT("SteamVRBindings");
 
 				TSharedPtr<FJsonObject> DescriptionsObject = MakeShareable(new FJsonObject);
 
@@ -1359,10 +1363,45 @@ private:
 
 				TArray<TSharedPtr<FJsonValue>> DefaultBindings;
 				{
-					TSharedRef<FJsonObject> BindingObject = MakeShareable(new FJsonObject());
-					BindingObject->SetStringField(TEXT("controller_type"), TEXT("vive_controller"));
-					BindingObject->SetStringField(TEXT("binding_url"), TEXT("bindings_vive_controller.json"));
-					DefaultBindings.Add(MakeShareable(new FJsonValueObject(BindingObject)));
+					IFileManager& FileManager = FFileManagerGeneric::Get();
+
+					// Find any default bindings stored in the project bindings dir.
+					// They must be saved as <PROJECT_CONFIG_DIR>/SteamVRBindings/<CONTROLLER_TYPE>.json in order to be included in the manifest.
+					TArray<FString> FoundFiles;
+					FileManager.FindFiles(FoundFiles, *BindingsDir, TEXT("*.json"));
+					UE_LOG(LogSteamVRController, Log, TEXT("Searching for bindings files in %s"), *BindingsDir);
+					for (FString& File : FoundFiles)
+					{
+						FString ControllerType = File;
+						ControllerType.RemoveFromEnd(".json");
+						TSharedRef<FJsonObject> BindingObject = MakeShareable(new FJsonObject());
+						BindingObject->SetStringField(TEXT("controller_type"), *ControllerType);
+						BindingObject->SetStringField(TEXT("binding_url"), *FileManager.ConvertToAbsolutePathForExternalAppForRead(*(BindingsDir / File)));
+						DefaultBindings.Add(MakeShareable(new FJsonValueObject(BindingObject)));
+					}
+
+					if (DefaultBindings.Num() == 0)
+					{
+#if WITH_EDITOR
+						UE_LOG(LogSteamVRController, Log, TEXT("No default Steam VR Input bindings found. Adding an inital entry for vive_controller."));
+						
+						// Create the directory if it doesn't exist.
+						if (!FileManager.DirectoryExists(*BindingsDir))
+						{
+							FileManager.MakeDirectory(*BindingsDir);
+						}
+
+						// Telling SteamVR about the path will allow the bindings editor to automatically export the file to the correct location when
+						// pressing "Replace Default Binding".
+						FString DefaultViveController = FileManager.ConvertToAbsolutePathForExternalAppForRead(*(BindingsDir / TEXT("vive_controller.json")));
+						TSharedRef<FJsonObject> BindingObject = MakeShareable(new FJsonObject());
+						BindingObject->SetStringField(TEXT("controller_type"), TEXT("vive_controller"));
+						BindingObject->SetStringField(TEXT("binding_url"), *DefaultViveController);
+						DefaultBindings.Add(MakeShareable(new FJsonValueObject(BindingObject)));
+#else
+						UE_LOG(LogSteamVRController, Error, TEXT("No default Steam VR Input bindings found in %s."), *BindingsDir);
+#endif
+					}
 				}
 
 				TArray<TSharedPtr<FJsonValue>> ActionSets;
