UE-57680 [CrashReport] UE4Editor_SteamVR!FSteamVRModel::GetRawMeshData() [steamvrassetmanager.cpp:307]
-Added some more logging to try to figure out why this is crashing.  Logs the resourceID of the steamVRMesh which might tell us if this is happening with some specific hardware.  Also logs index and vertex counts and pointers to see if the data is somewhat malformed or just completely invalid.

diff --git a/Engine/Plugins/Runtime/Steam/SteamVR/Source/SteamVR/Private/SteamVRAssetManager.cpp b/Engine/Plugins/Runtime/Steam/SteamVR/Source/SteamVR/Private/SteamVRAssetManager.cpp
--- a/Engine/Plugins/Runtime/Steam/SteamVR/Source/SteamVR/Private/SteamVRAssetManager.cpp
+++ b/Engine/Plugins/Runtime/Steam/SteamVR/Source/SteamVR/Private/SteamVRAssetManager.cpp
@@ -286,6 +286,9 @@ bool FSteamVRModel::GetRawMeshData(float UEMeterScale, FSteamVRMeshData& MeshDat
 #if STEAMVR_SUPPORTED_PLATFORMS
 	if (RawResource)
 	{
+		// Logging to try to get info about a crash where the RawMeshData seems to be bad.
+		UE_LOG(LogSteamVR, Log, TEXT("FSteamVRModel::GetRawMeshData ResourceId %s rVertexData 0x%lx unVertexCount %i rIndexData 0x%lx unTriangleCount %i diffuseTextureId %i"), *ResourceId, RawResource->rVertexData, RawResource->unVertexCount, RawResource->rIndexData, RawResource->unTriangleCount, RawResource->diffuseTextureId);
+
 		const uint32 VertCount = RawResource->unVertexCount;
 		MeshDataOut.VertPositions.Empty(VertCount);
 		MeshDataOut.UVs.Empty(VertCount);
