UE-57680 [CrashReport] UE4Editor_SteamVR!FSteamVRModel::GetRawMeshData() [steamvrassetmanager.cpp:307]
-Two more log lines for steamvr shutdown.  Maybe we are trying to get the models after shutdown and get bad data?

diff --git a/Engine/Plugins/Runtime/Steam/SteamVR/Source/SteamVR/Private/SteamVRHMD.cpp b/Engine/Plugins/Runtime/Steam/SteamVR/Source/SteamVR/Private/SteamVRHMD.cpp
--- a/Engine/Plugins/Runtime/Steam/SteamVR/Source/SteamVR/Private/SteamVRHMD.cpp
+++ b/Engine/Plugins/Runtime/Steam/SteamVR/Source/SteamVR/Private/SteamVRHMD.cpp
@@ -222,6 +222,7 @@ public:
 	{
 		if (OpenVRDLLHandle != nullptr)
 		{
+			UE_LOG(LogHMD, Log, TEXT("Freeing OpenVRDLLHandle."));
 			FPlatformProcess::FreeDllHandle(OpenVRDLLHandle);
 			OpenVRDLLHandle = nullptr;
 		}
@@ -1706,6 +1707,8 @@ void FSteamVRHMD::Shutdown()
 {
 	if (VRSystem != nullptr)
 	{
+		UE_LOG(LogHMD, Log, TEXT("SteamVR Shutting down."));
+
 		SplashTicker->UnregisterForMapLoad();
 		SplashTicker = nullptr;
 
