Fix a variable rename in CL 4671063

diff --git a/Engine/Plugins/Runtime/Steam/SteamVR/Source/SteamVR/Private/SteamVRHMD.cpp b/Engine/Plugins/Runtime/Steam/SteamVR/Source/SteamVR/Private/SteamVRHMD.cpp
--- a/Engine/Plugins/Runtime/Steam/SteamVR/Source/SteamVR/Private/SteamVRHMD.cpp
+++ b/Engine/Plugins/Runtime/Steam/SteamVR/Source/SteamVR/Private/SteamVRHMD.cpp
@@ -249,7 +249,7 @@ public:
 			// It's only triggered in cases where SteamVR is available, but was not running prior to calling VR_Init above.
 			if ((--CompositorConnectAttempts > 0) && (VRInitErr == vr::VRInitError_IPC_CompositorConnectFailed || VRInitErr == vr::VRInitError_IPC_CompositorInvalidConnectResponse))
 			{
-				UE_LOG(LogHMD, Warning, TEXT("Failed to get Compositor connnection (%d) retrying... (%d attempt(s) left)"), (int32)VRInitErr, CompositorConnectRetries);
+				UE_LOG(LogHMD, Warning, TEXT("Failed to get Compositor connnection (%d) retrying... (%d attempt(s) left)"), (int32)VRInitErr, CompositorConnectAttempts);
 				FPlatformProcess::Sleep(1);
 			}
 			else
