Don't attempt to initialize SteamVRHMD when the compositor is not connected

diff --git a/Engine/Plugins/Runtime/Steam/SteamVR/Source/SteamVR/Private/SteamVRHMD.cpp b/Engine/Plugins/Runtime/Steam/SteamVR/Source/SteamVR/Private/SteamVRHMD.cpp
--- a/Engine/Plugins/Runtime/Steam/SteamVR/Source/SteamVR/Private/SteamVRHMD.cpp
+++ b/Engine/Plugins/Runtime/Steam/SteamVR/Source/SteamVR/Private/SteamVRHMD.cpp
@@ -412,9 +412,9 @@ IMPLEMENT_MODULE( FSteamVRPlugin, SteamVR )
 TSharedPtr< class IXRTrackingSystem, ESPMode::ThreadSafe > FSteamVRPlugin::CreateTrackingSystem()
 {
 #if STEAMVR_SUPPORTED_PLATFORMS
-	if (!VRSystem)
+	if (!VRSystem && !Initialize())
 	{
-		Initialize();
+		return nullptr;
 	}
 
 	TSharedPtr< FSteamVRHMD, ESPMode::ThreadSafe > SteamVRHMD = FSceneViewExtensions::NewExtension<FSteamVRHMD>(this);
