Fix compiler warnings in SteamVRHMD

diff --git a/Engine/Plugins/Runtime/Steam/SteamVR/Source/SteamVR/Private/SteamVRHMD.cpp b/Engine/Plugins/Runtime/Steam/SteamVR/Source/SteamVR/Private/SteamVRHMD.cpp
--- a/Engine/Plugins/Runtime/Steam/SteamVR/Source/SteamVR/Private/SteamVRHMD.cpp
+++ b/Engine/Plugins/Runtime/Steam/SteamVR/Source/SteamVR/Private/SteamVRHMD.cpp
@@ -1563,11 +1563,11 @@ FSteamVRHMD::FSteamVRHMD(const FAutoRegister& AutoRegister, ISteamVRPlugin* InSt
 	RendererModule(nullptr),
 	SteamVRPlugin(InSteamVRPlugin),
 	VRSystem(vr::VRSystem()),
-	VRInput(vr::VRInput()),
 	VRCompositor(vr::VRCompositor()),
 	VROverlay(vr::VROverlay()),
 	VRChaperone(vr::VRChaperone()),
-	VRRenderModels(vr::VRRenderModels())
+	VRRenderModels(vr::VRRenderModels()),
+	VRInput(vr::VRInput())
 {
 	Startup();
 }
diff --git a/Engine/Plugins/Runtime/Steam/SteamVR/Source/SteamVR/Private/SteamVRHMD.h b/Engine/Plugins/Runtime/Steam/SteamVR/Source/SteamVR/Private/SteamVRHMD.h
--- a/Engine/Plugins/Runtime/Steam/SteamVR/Source/SteamVR/Private/SteamVRHMD.h
+++ b/Engine/Plugins/Runtime/Steam/SteamVR/Source/SteamVR/Private/SteamVRHMD.h
@@ -312,6 +312,7 @@ public:
 	{
 	public:
 		FVulkanExtensions(vr::IVRCompositor* InVRCompositor);
+		virtual ~FVulkanExtensions() {}
 
 		/** IHeadMountedDisplayVulkanExtensions */
 		virtual bool GetVulkanInstanceExtensionsRequired(TArray<const ANSICHAR*>& Out) override;
