Fixed vulkan compilation errors on Mac for Steam VR

diff --git a/Engine/Plugins/Runtime/Steam/SteamVR/Source/SteamVR/Private/SteamVRHMD.cpp b/Engine/Plugins/Runtime/Steam/SteamVR/Source/SteamVR/Private/SteamVRHMD.cpp
--- a/Engine/Plugins/Runtime/Steam/SteamVR/Source/SteamVR/Private/SteamVRHMD.cpp
+++ b/Engine/Plugins/Runtime/Steam/SteamVR/Source/SteamVR/Private/SteamVRHMD.cpp
@@ -382,7 +382,7 @@ public:
 
 	virtual TSharedPtr< IHeadMountedDisplayVulkanExtensions, ESPMode::ThreadSafe > GetVulkanExtensions() override
 	{
-#if STEAMVR_SUPPORTED_PLATFORMS
+#if STEAMVR_SUPPORTED_PLATFORMS && !PLATFORM_MAC
 		if (Initialize())
 		{
 			if (!VulkanExtensions.IsValid())
@@ -433,6 +433,8 @@ TSharedPtr< class IXRTrackingSystem, ESPMode::ThreadSafe > FSteamVRPlugin::Creat
 
 #if STEAMVR_SUPPORTED_PLATFORMS
 
+#if !PLATFORM_MAC
+
 FSteamVRHMD::FVulkanExtensions::FVulkanExtensions(vr::IVRCompositor* InVRCompositor)
 	: VRCompositor(InVRCompositor)
 {
@@ -504,6 +506,8 @@ bool FSteamVRHMD::FVulkanExtensions::GetVulkanDeviceExtensionsRequired(VkPhysica
 	return true;
 }
 
+#endif
+
 bool FSteamVRHMD::IsHMDConnected()
 {
 	return SteamVRPlugin->IsHMDConnected();
