Reset OpenVR if initialization fails.

diff --git a/Engine/Plugins/Runtime/Steam/SteamVR/Source/SteamVR/Private/SteamVRHMD.cpp b/Engine/Plugins/Runtime/Steam/SteamVR/Source/SteamVR/Private/SteamVRHMD.cpp
--- a/Engine/Plugins/Runtime/Steam/SteamVR/Source/SteamVR/Private/SteamVRHMD.cpp
+++ b/Engine/Plugins/Runtime/Steam/SteamVR/Source/SteamVR/Private/SteamVRHMD.cpp
@@ -234,6 +234,7 @@ public:
 		if ((VRSystem == nullptr) || (VRInitErr != vr::VRInitError_None))
 		{
 			UE_LOG(LogHMD, Log, TEXT("Failed to initialize OpenVR (version mismatch) with code %d"), (int32)VRInitErr);
+			Reset();
 			return false;
 		}
 
@@ -261,6 +262,7 @@ public:
 		if (VRInitErr != vr::VRInitError_None)
 		{
 			UE_LOG(LogHMD, Log, TEXT("SteamVR failed to initialize.  Error: %d"), (int32)VRInitErr);
+			Reset();
 			return false;
 		}
 
