PR #5698: Fixed texture colour order for FSteamVRTexture, plus extended to work outside of editor (Contributed by [REDACTED_PERSON])

diff --git a/Engine/Plugins/Runtime/Steam/SteamVR/Source/SteamVR/Private/SteamVRAssetManager.cpp b/Engine/Plugins/Runtime/Steam/SteamVR/Source/SteamVR/Private/SteamVRAssetManager.cpp
--- a/Engine/Plugins/Runtime/Steam/SteamVR/Source/SteamVR/Private/SteamVRAssetManager.cpp
+++ b/Engine/Plugins/Runtime/Steam/SteamVR/Source/SteamVR/Private/SteamVRAssetManager.cpp
@@ -333,24 +333,29 @@ public:
 	int32 GetResourceID() const { return ResourceId; }
 
 public:
-	UTexture2D* ConstructUETexture(UObject* ObjOuter, const FName ObjName, EObjectFlags ObjFlags = RF_NoFlags)
+	UTexture2D* ConstructUETexture(const FName ObjName)
 	{
 		UTexture2D* NewTexture = nullptr;
 
 #if STEAMVR_SUPPORTED_PLATFORMS
 		if (RawResource != nullptr)
 		{
-#if WITH_EDITORONLY_DATA // @TODO: UTexture::Source is only available in editor builds, we need to find some other way to construct textures - try using CreateTransient() (see: TexturePaintHelpers::CreateTempUncompressedTexture)
-			NewTexture = NewObject<UTexture2D>(ObjOuter, ObjName, ObjFlags);
-			NewTexture->Source.Init(RawResource->unWidth, RawResource->unHeight, /*NewNumSlices =*/1, /*NewNumMips =*/1, TSF_BGRA8, RawResource->rubTextureMapData);
-
-			NewTexture->MipGenSettings = TMGS_NoMipmaps;
-			// disable compression
-			NewTexture->CompressionNone = true;
-			NewTexture->DeferCompression = false;
-
-			NewTexture->PostEditChange();
-#endif 
+			// Create the texture. Using UTexture2D::CreateTransient which is supported outside of editor builds.
+			NewTexture = UTexture2D::CreateTransient(RawResource->unWidth, RawResource->unHeight, EPixelFormat::PF_R8G8B8A8, ObjName);
+			if (NewTexture != nullptr)
+			{
+				FTexture2DMipMap& Mip = NewTexture->PlatformData->Mips[0];
+				void* Data = Mip.BulkData.Lock(LOCK_READ_WRITE);
+				FMemory::Memcpy(Data, (void*)RawResource->rubTextureMapData, Mip.BulkData.GetBulkDataSize());
+				Mip.BulkData.Unlock();
+				NewTexture->PlatformData->NumSlices = 1;
+#if WITH_EDITORONLY_DATA
+				NewTexture->CompressionNone = true;
+				NewTexture->DeferCompression = false;
+				NewTexture->MipGenSettings = TMGS_NoMipmaps;
+#endif				
+				NewTexture->UpdateResource();
+			}
 		}
 #endif
 
@@ -491,13 +496,12 @@ void FSteamVRAsyncMeshLoader::Tick(float /*DeltaTime*/)
 
 				if (bLoadSuccess)
 				{
-					UObject* TextureOuter = GetTransientPackage();
 					FName    TextureName  = *FString::Printf(TEXT("T_SteamVR_%d"), TextureResource.GetResourceID());
 
-					UTexture2D* UETexture = FindObjectFast<UTexture2D>(TextureOuter, TextureName, /*ExactClass =*/true);
+					UTexture2D* UETexture = FindObjectFast<UTexture2D>(GetTransientPackage(), TextureName, /*ExactClass =*/true);
 					if (UETexture == nullptr)
 					{
-						UETexture = TextureResource.ConstructUETexture(TextureOuter, TextureName);
+						UETexture = TextureResource.ConstructUETexture(TextureName);
 					}
 					ConstructedTextures.Add(TextureResource.GetResourceID(), UETexture);
 				}
diff --git a/Engine/Source/Runtime/Engine/Classes/Engine/Texture2D.h b/Engine/Source/Runtime/Engine/Classes/Engine/Texture2D.h
--- a/Engine/Source/Runtime/Engine/Classes/Engine/Texture2D.h
+++ b/Engine/Source/Runtime/Engine/Classes/Engine/Texture2D.h
@@ -369,7 +369,7 @@ public:
 	friend struct FStreamingRenderAsset;
 	
 	/** creates and initializes a new Texture2D with the requested settings */
-	ENGINE_API static class UTexture2D* CreateTransient(int32 InSizeX, int32 InSizeY, EPixelFormat InFormat = PF_B8G8R8A8);
+	ENGINE_API static class UTexture2D* CreateTransient(int32 InSizeX, int32 InSizeY, EPixelFormat InFormat = PF_B8G8R8A8, const FName InName = NAME_None);
 
 	/**
 	 * Gets the X size of the texture, in pixels
diff --git a/Engine/Source/Runtime/Engine/Private/Texture2D.cpp b/Engine/Source/Runtime/Engine/Private/Texture2D.cpp
--- a/Engine/Source/Runtime/Engine/Private/Texture2D.cpp
+++ b/Engine/Source/Runtime/Engine/Private/Texture2D.cpp
@@ -1192,7 +1192,7 @@ bool UTexture2D::IsFullyStreamedIn()
 }
 
 
-UTexture2D* UTexture2D::CreateTransient(int32 InSizeX, int32 InSizeY, EPixelFormat InFormat)
+UTexture2D* UTexture2D::CreateTransient(int32 InSizeX, int32 InSizeY, EPixelFormat InFormat, const FName InName)
 {
 	LLM_SCOPE(ELLMTag::Textures);
 
@@ -1203,7 +1203,7 @@ UTexture2D* UTexture2D::CreateTransient(int32 InSizeX, int32 InSizeY, EPixelForm
 	{
 		NewTexture = NewObject<UTexture2D>(
 			GetTransientPackage(),
-			NAME_None,
+			InName,
 			RF_Transient
 			);
 
