Fix potential concurrency issues with FSteamVRHMD::GetAllocatedTexture

diff --git a/Engine/Plugins/Runtime/Steam/SteamVR/Source/SteamVR/Private/SteamVRStereoLayers.cpp b/Engine/Plugins/Runtime/Steam/SteamVR/Source/SteamVR/Private/SteamVRStereoLayers.cpp
--- a/Engine/Plugins/Runtime/Steam/SteamVR/Source/SteamVR/Private/SteamVRStereoLayers.cpp
+++ b/Engine/Plugins/Runtime/Steam/SteamVR/Source/SteamVR/Private/SteamVRStereoLayers.cpp
@@ -330,43 +330,30 @@ void FSteamVRHMD::GetAllocatedTexture(uint32 LayerId, FTextureRHIRef &Texture, F
 {
 	Texture = LeftTexture = nullptr;
 	FSteamVRLayer* LayerFound = nullptr;
+	check(IsInRenderingThread()); // Not strictly necessary, as WithLayer uses a scope lock
 
-	if (IsInRenderingThread())
+	WithLayer(LayerId, [&](FSteamVRLayer* LayerFound)
 	{
-		ForEachLayer([&](uint32 /* unused */, FSteamVRLayer& Layer)
+		if (LayerFound && LayerFound->LayerDesc.Texture)
 		{
-			if (Layer.GetLayerId() == LayerId)
+			switch (LayerFound->LayerDesc.ShapeType)
 			{
-				LayerFound = &Layer;
+			case IStereoLayers::CubemapLayer:
+				Texture = LayerFound->LayerDesc.Texture->GetTextureCube();
+				LeftTexture = LayerFound->LayerDesc.LeftTexture ? LayerFound->LayerDesc.LeftTexture->GetTextureCube() : nullptr;
+				break;
+
+			case IStereoLayers::CylinderLayer:
+			case IStereoLayers::QuadLayer:
+				Texture = LayerFound->LayerDesc.Texture->GetTexture2D();
+				LeftTexture = LayerFound->LayerDesc.LeftTexture ? LayerFound->LayerDesc.LeftTexture->GetTexture2D() : nullptr;
+				break;
+
+			default:
+				break;
 			}
-		});
-	}
-	else
-	{
-		// Only supporting the use of this function on RenderingThread.
-		check(false);
-		return;
-	}
-
-	if (LayerFound && LayerFound->LayerDesc.Texture)
-	{
-		switch (LayerFound->LayerDesc.ShapeType)
-		{
-		case IStereoLayers::CubemapLayer:
-			Texture = LayerFound->LayerDesc.Texture->GetTextureCube();
-			LeftTexture = LayerFound->LayerDesc.LeftTexture ? LayerFound->LayerDesc.LeftTexture->GetTextureCube() : nullptr;
-			break;
-
-		case IStereoLayers::CylinderLayer:
-		case IStereoLayers::QuadLayer:
-			Texture = LayerFound->LayerDesc.Texture->GetTexture2D();
-			LeftTexture = LayerFound->LayerDesc.LeftTexture ? LayerFound->LayerDesc.LeftTexture->GetTexture2D() : nullptr;
-			break;
-
-		default:
-			break;
 		}
-	}
+	});
 }
 
 //=============================================================================
diff --git a/Engine/Source/Runtime/HeadMountedDisplay/Public/StereoLayerManager.h b/Engine/Source/Runtime/HeadMountedDisplay/Public/StereoLayerManager.h
--- a/Engine/Source/Runtime/HeadMountedDisplay/Public/StereoLayerManager.h
+++ b/Engine/Source/Runtime/HeadMountedDisplay/Public/StereoLayerManager.h
@@ -85,6 +85,12 @@ protected:
 		}
 	}
 
+	void WithLayer(uint32 LayerId, TFunction<void(LayerType*)> Func)
+	{
+		FScopeLock LockLayers(&LayerCritSect);
+		Func(StereoLayers.Find(LayerId));
+	}
+
 public:
 
 	TStereoLayerManager()
