Assert in SteamVR when autoshowing splash screens

diff --git a/Engine/Plugins/Runtime/Steam/SteamVR/Source/SteamVR/Private/SteamVRSplash.cpp b/Engine/Plugins/Runtime/Steam/SteamVR/Source/SteamVR/Private/SteamVRSplash.cpp
--- a/Engine/Plugins/Runtime/Steam/SteamVR/Source/SteamVR/Private/SteamVRSplash.cpp
+++ b/Engine/Plugins/Runtime/Steam/SteamVR/Source/SteamVR/Private/SteamVRSplash.cpp
@@ -51,7 +51,7 @@ void FSteamSplashTicker::Tick(float DeltaTime)
 	// even though when used by the renderer as an indication whether normal present is needed.
 	if (SteamVRHMD->bSplashIsShown && SteamVRHMD->pBridge && SteamVRHMD->pBridge->Present(Dummy))
 	{
-		check(!SteamVRHMD->VRCompositor);
+		check(SteamVRHMD->VRCompositor);
 		SteamVRHMD->VRCompositor->PostPresentHandoff();
 	}
 }
