Making the supported platforms for Oculus and SteamVR explicit to ensure they don't interfere with HL2.

diff --git a/Engine/Plugins/Runtime/Oculus/OculusVR/OculusVR.uplugin b/Engine/Plugins/Runtime/Oculus/OculusVR/OculusVR.uplugin
--- a/Engine/Plugins/Runtime/Oculus/OculusVR/OculusVR.uplugin
+++ b/Engine/Plugins/Runtime/Oculus/OculusVR/OculusVR.uplugin
@@ -14,6 +14,11 @@
 	"CanContainContent": true,
 	"IsBetaVersion": false,
 	"Installed": false,
+	"SupportedTargetPlatforms": [
+		"Win32",
+		"Win64",
+		"Android"
+	],
 	"Modules": [
 		{
 			"Name": "OculusHMD",
diff --git a/Engine/Plugins/Runtime/Steam/SteamVR/SteamVR.uplugin b/Engine/Plugins/Runtime/Steam/SteamVR/SteamVR.uplugin
--- a/Engine/Plugins/Runtime/Steam/SteamVR/SteamVR.uplugin
+++ b/Engine/Plugins/Runtime/Steam/SteamVR/SteamVR.uplugin
@@ -1,27 +1,31 @@
 {
-	"FileVersion" : 3,
-	"Version" : 1,
-	"VersionName" : "1.0",
-	"FriendlyName" : "SteamVR",
-	"Description" : "Support for the SteamVR API",
-	"Category" : "Virtual Reality",
-	"CreatedBy" : "Epic Games, Inc.",
-	"CreatedByURL" : "http://epicgames.com",
-	"DocsURL" : "",
-	"MarketplaceURL" : "",
-	"SupportURL" : "",
-	"EnabledByDefault" : true,
-	"CanContainContent" : true,
-	"IsBetaVersion" : false,
-	"Installed" : false,
-	"Modules" :
-	[
+	"FileVersion": 3,
+	"Version": 1,
+	"VersionName": "1.0",
+	"FriendlyName": "SteamVR",
+	"Description": "Support for the SteamVR API",
+	"Category": "Virtual Reality",
+	"CreatedBy": "Epic Games, Inc.",
+	"CreatedByURL": "http://epicgames.com",
+	"DocsURL": "",
+	"MarketplaceURL": "",
+	"SupportURL": "",
+	"EnabledByDefault": true,
+	"CanContainContent": true,
+	"IsBetaVersion": false,
+	"Installed": false,
+	"SupportedTargetPlatforms": [
+		"Win32",
+		"Win64",
+		"Linux",
+		"Mac"
+	],
+	"Modules": [
 		{
-			"Name" : "SteamVR",
-			"Type" : "Runtime",
-			"LoadingPhase" : "PostConfigInit",
-			"WhitelistPlatforms" :
-			[
+			"Name": "SteamVR",
+			"Type": "Runtime",
+			"LoadingPhase": "PostConfigInit",
+			"WhitelistPlatforms": [
 				"Win64",
 				"Win32",
 				"Linux",
@@ -29,11 +33,10 @@
 			]
 		},
 		{
-			"Name" : "SteamVRController",
-			"Type" : "Runtime",
-			"LoadingPhase" : "Default",
-			"WhitelistPlatforms" :
-			[
+			"Name": "SteamVRController",
+			"Type": "Runtime",
+			"LoadingPhase": "Default",
+			"WhitelistPlatforms": [
 				"Win64",
 				"Win32",
 				"Linux",
@@ -41,8 +44,7 @@
 			]
 		}
 	],
-	"Plugins":
-	[
+	"Plugins": [
 		{
 			"Name": "ProceduralMeshComponent",
 			"Enabled": true
