Plumb depth-textures through for SteamVR so Oculus overlays work correctly.

diff --git a/Engine/Plugins/Runtime/Steam/SteamVR/Source/SteamVR/Private/SteamVRHMD.cpp b/Engine/Plugins/Runtime/Steam/SteamVR/Source/SteamVR/Private/SteamVRHMD.cpp
--- a/Engine/Plugins/Runtime/Steam/SteamVR/Source/SteamVR/Private/SteamVRHMD.cpp
+++ b/Engine/Plugins/Runtime/Steam/SteamVR/Source/SteamVR/Private/SteamVRHMD.cpp
@@ -916,47 +916,6 @@ private:
 };
 #endif
 
-bool FSteamVRHMD::AllocateRenderTargetTexture(uint32 Index, uint32 SizeX, uint32 SizeY, uint8 Format, uint32 NumMips, uint32 InTexFlags, uint32 InTargetableTextureFlags, FTexture2DRHIRef& OutTargetableTexture, FTexture2DRHIRef& OutShaderResourceTexture, uint32 NumSamples)
-{
-	if (!IsStereoEnabled())
-	{
-		return false;
-	}
-	
-#if PLATFORM_MAC
-	const uint32 SwapChainLength = 3;
-	
-	MetalBridge* MetalBridgePtr = (MetalBridge*)pBridge.GetReference();
-	
-	MetalBridgePtr->TextureSet = new FRHITextureSet2D(SwapChainLength, PF_B8G8R8A8, SizeX, SizeY, 1, NumSamples, InTexFlags, FClearValueBinding(FLinearColor::Transparent));
-	
-	for (uint32 SwapChainIter = 0; SwapChainIter < SwapChainLength; ++SwapChainIter)
-	{
-		IOSurfaceRef Surface = MetalBridgePtr->GetSurface(SizeX, SizeY);
-		check(Surface != nil);
-		
-		FRHIResourceCreateInfo CreateInfo;
-		CreateInfo.BulkData = new FIOSurfaceResourceWrapper(Surface);
-		CFRelease(Surface);
-		CreateInfo.ResourceArray = nullptr;
-		
-		FTexture2DRHIRef TargetableTexture, ShaderResourceTexture;
-		RHICreateTargetableShaderResource2D(SizeX, SizeY, PF_B8G8R8A8, 1, TexCreate_None, TexCreate_RenderTargetable, false, CreateInfo, TargetableTexture, ShaderResourceTexture, NumSamples);
-		check(TargetableTexture == ShaderResourceTexture);
-		static_cast<FRHITextureSet2D*>(MetalBridgePtr->TextureSet.GetReference())->AddTexture(TargetableTexture, SwapChainIter);
-	}
-	
-	OutTargetableTexture = OutShaderResourceTexture = MetalBridgePtr->TextureSet;
-
-	return true;
-#else
-	FRHIResourceCreateInfo CreateInfo;
-	RHICreateTargetableShaderResource2D(SizeX, SizeY, PF_B8G8R8A8, 1, TexCreate_None, TexCreate_RenderTargetable, false, CreateInfo, OutTargetableTexture, OutShaderResourceTexture);
-
-	return true;
-#endif
-}
-
 void FSteamVRHMD::PoseToOrientationAndPosition(const vr::HmdMatrix34_t& InPose, const float WorldToMetersScale, FQuat& OutOrientation, FVector& OutPosition) const
 {
 	FMatrix Pose = ToFMatrix(InPose);
@@ -1467,7 +1426,6 @@ FMatrix FSteamVRHMD::GetStereoProjectionMatrix(const enum EStereoscopicPass Ster
 #endif
 
 	return Mat;
-
 }
 
 void FSteamVRHMD::GetEyeRenderParams_RenderThread(const FRenderingCompositePassContext& Context, FVector2D& EyeToSrcUVScaleValue, FVector2D& EyeToSrcUVOffsetValue) const
@@ -1560,6 +1518,160 @@ bool FSteamVRHMD::NeedReAllocateViewportRenderTarget(const FViewport& Viewport)
 	return false;
 }
 
+bool FSteamVRHMD::NeedReAllocateDepthTexture(const TRefCountPtr<struct IPooledRenderTarget> & DepthTarget)
+{
+	check(IsInGameThread());
+
+	// Check the dimensions of the currently stored depth swapchain vs the current rendering swapchain.
+	if (pBridge->GetSwapChain()->GetTexture2D()->GetSizeX() != pBridge->GetDepthSwapChain()->GetTexture2D()->GetSizeX() ||
+		pBridge->GetSwapChain()->GetTexture2D()->GetSizeY() != pBridge->GetDepthSwapChain()->GetTexture2D()->GetSizeY())
+	{
+		return true;
+	}
+
+	return false;
+}
+
+#if PLATFORM_MAC
+static const uint32 SteamVRSwapChainLength = 3;
+#else
+static const uint32 SteamVRSwapChainLength = 1;
+#endif
+
+bool FSteamVRHMD::AllocateRenderTargetTexture(uint32 Index, uint32 SizeX, uint32 SizeY, uint8 Format, uint32 NumMips, uint32 InTexFlags, uint32 InTargetableTextureFlags, FTexture2DRHIRef& OutTargetableTexture, FTexture2DRHIRef& OutShaderResourceTexture, uint32 NumSamples)
+{
+	if (!IsStereoEnabled())
+	{
+		return false;
+	}
+
+	TArray<FTextureRHIRef> SwapChainTextures;
+	FTextureRHIRef BindingTexture;
+
+#if PLATFORM_MAC
+	MetalBridge* MetalBridgePtr = (MetalBridge*)pBridge.GetReference();
+#endif
+
+	for (uint32 SwapChainIter = 0; SwapChainIter < SteamVRSwapChainLength; ++SwapChainIter)
+	{
+#if PLATFORM_MAC
+		IOSurfaceRef Surface = MetalBridgePtr->GetSurface(SizeX, SizeY);
+		check(Surface != nil);
+
+		FRHIResourceCreateInfo CreateInfo;
+		CreateInfo.BulkData = new FIOSurfaceResourceWrapper(Surface);
+		CFRelease(Surface);
+		CreateInfo.ResourceArray = nullptr;
+#else
+		FRHIResourceCreateInfo CreateInfo;
+#endif
+
+		FTexture2DRHIRef TargetableTexture, ShaderResourceTexture;
+		RHICreateTargetableShaderResource2D(SizeX, SizeY, PF_B8G8R8A8, 1, TexCreate_None, TexCreate_RenderTargetable, false, CreateInfo, TargetableTexture, ShaderResourceTexture, NumSamples);
+		check(TargetableTexture == ShaderResourceTexture);
+
+		SwapChainTextures.Add((FTextureRHIRef&)TargetableTexture);
+
+		if (BindingTexture == nullptr)
+		{
+			// For this iteration, create a temporary texture resource that will get stomped the first time we do our bind/alias
+			// pass on this object.
+			// Note, even though this is a wasted allocation, it never gets bound or used in any way, so should be minimally impactful.
+			FTexture2DRHIRef TempTargetableTexture, TempShaderResourceTexture;
+			RHICreateTargetableShaderResource2D(SizeX, SizeY, PF_B8G8R8A8, 1, TexCreate_None, TexCreate_RenderTargetable, false, CreateInfo, TempTargetableTexture, TempShaderResourceTexture, NumSamples);
+			check(TempTargetableTexture == TempShaderResourceTexture);
+			BindingTexture = TempTargetableTexture;
+
+			// Alias immediately.
+			GDynamicRHI->RHIAliasTextureResources(BindingTexture, TargetableTexture);
+		}
+	}
+
+	pBridge->CreateSwapChain(BindingTexture, MoveTemp(SwapChainTextures));
+
+	// These are the same.
+	OutTargetableTexture = (FTexture2DRHIRef&)BindingTexture;
+	OutShaderResourceTexture = (FTexture2DRHIRef&)BindingTexture;
+
+	return true;
+}
+
+bool FSteamVRHMD::AllocateDepthTexture(uint32 Index, uint32 SizeX, uint32 SizeY, uint8 Format, uint32 NumMips, uint32 Flags, uint32 TargetableTextureFlags, 
+	FTexture2DRHIRef& OutTargetableTexture, FTexture2DRHIRef& OutShaderResourceTexture, uint32 NumSamples /* = 1 */)
+{
+	if (!IsStereoEnabled() || pBridge == nullptr)
+	{
+		return false;
+	}
+
+#if PLATFORM_MAC
+	// @todo: Determine if we want to manage depth on the Mac?
+	return false;
+#else
+	auto DepthSwapChain = pBridge->GetDepthSwapChain();
+	if (DepthSwapChain)
+	{
+		// If size is the same as requested, return the swap-chain texture.
+		if (DepthSwapChain->GetTexture2D()->GetSizeX() == SizeX && DepthSwapChain->GetTexture2D()->GetSizeY() == SizeY)
+		{
+			// @todo: Do we need to check format, etc?
+			OutTargetableTexture = OutShaderResourceTexture = DepthSwapChain->GetTexture2D();
+			return true;
+		}
+	}
+
+	TArray<FTextureRHIRef> SwapChainTextures;
+	FTextureRHIRef BindingTexture;
+
+	FClearValueBinding ClearValue(0.0f, 0);
+	ClearValue.ColorBinding = EClearBinding::EDepthStencilBound;
+	FRHIResourceCreateInfo CreateInfo(ClearValue);
+	CreateInfo.DebugName = _TEXT("SteamVRDepthStencil");
+
+	for (uint32 SwapChainIter = 0; SwapChainIter < SteamVRSwapChainLength; ++SwapChainIter)
+	{
+		FTexture2DRHIRef TargetableTexture, ShaderResourceTexture;
+
+		RHICreateTargetableShaderResource2D(
+			SizeX,
+			SizeY,
+			PF_DepthStencil,
+			1,			// 1 mip for depth (0 is usually passed in above, which is invalid).
+			Flags,
+			TargetableTextureFlags,
+			false,
+			CreateInfo,
+			TargetableTexture,
+			ShaderResourceTexture,
+			NumSamples);
+
+		check(TargetableTexture == ShaderResourceTexture);
+
+		SwapChainTextures.Add((FTextureRHIRef&)TargetableTexture);
+
+		// As above, create a temp, stompable texture to alias over.
+		if (BindingTexture == nullptr)
+		{
+			FTexture2DRHIRef TempTargetableTexture, TempShaderResourceTexture;
+			RHICreateTargetableShaderResource2D(SizeX, SizeY, PF_DepthStencil, 1, Flags, TargetableTextureFlags, false, CreateInfo, TempTargetableTexture, TempShaderResourceTexture, NumSamples);
+			check(TempTargetableTexture == TempShaderResourceTexture);
+			BindingTexture = TempTargetableTexture;
+
+			// Alias immediately.
+			GDynamicRHI->RHIAliasTextureResources(BindingTexture, TargetableTexture);
+		}
+	}
+
+	OutTargetableTexture = (FTexture2DRHIRef&)BindingTexture;
+	OutShaderResourceTexture = (FTexture2DRHIRef&)BindingTexture;
+
+	// Create the bridge's depth swapchain.
+	pBridge->CreateDepthSwapChain(BindingTexture, MoveTemp(SwapChainTextures));
+
+	return true;
+#endif
+}
+
 FSteamVRHMD::FSteamVRHMD(const FAutoRegister& AutoRegister, ISteamVRPlugin* InSteamVRPlugin) :
 	FHeadMountedDisplayBase(nullptr),
 	FSceneViewExtensionBase(AutoRegister), 
diff --git a/Engine/Plugins/Runtime/Steam/SteamVR/Source/SteamVR/Private/SteamVRHMD.h b/Engine/Plugins/Runtime/Steam/SteamVR/Source/SteamVR/Private/SteamVRHMD.h
--- a/Engine/Plugins/Runtime/Steam/SteamVR/Source/SteamVR/Private/SteamVRHMD.h
+++ b/Engine/Plugins/Runtime/Steam/SteamVR/Source/SteamVR/Private/SteamVRHMD.h
@@ -13,6 +13,7 @@
 #include "StereoLayerManager.h"
 #include "XRRenderTargetManager.h"
 #include "XRRenderBridge.h"
+#include "XRSwapChain.h"
 #include "IHeadMountedDisplayVulkanExtensions.h"
 
 #if PLATFORM_WINDOWS
@@ -66,52 +67,6 @@ struct FSteamVRLayer
 	friend void MarkLayerTextureForUpdate(FSteamVRLayer& Layer);
 };
 
-/**
- * Render target swap chain
- */
-class FRHITextureSet2D : public FRHITexture2D
-{
-public:
-	
-	FRHITextureSet2D(const uint32 TextureSetSize, EPixelFormat Format, uint32 SizeX, uint32 SizeY, uint32 NumMips, uint32 NumSamples, uint32 Flags, const FClearValueBinding& InClearValue)
-	: FRHITexture2D(SizeX, SizeY, NumMips, NumSamples, Format, Flags, InClearValue)
-	, TextureIndex(0)
-	{
-		TextureSet.AddZeroed(TextureSetSize);
-	}
-	
-	virtual ~FRHITextureSet2D()
-	{}
-	
-	void AddTexture(FTexture2DRHIRef& Texture, const uint32 Index)
-	{
-		check(Index < static_cast<uint32>(TextureSet.Num()));
-		// todo: Check texture format to ensure it matches the set
-		TextureSet[Index] = Texture;
-	}
-	
-	void Advance()
-	{
-		TextureIndex = (TextureIndex + 1) % static_cast<uint32>(TextureSet.Num());
-	}
-	
-	virtual void* GetTextureBaseRHI() override
-	{
-		check(TextureSet[TextureIndex].IsValid());
-		return TextureSet[TextureIndex]->GetTextureBaseRHI();
-	}
-	
-	virtual void* GetNativeResource() const override
-	{
-		check(TextureSet[TextureIndex].IsValid());
-		return TextureSet[TextureIndex]->GetNativeResource();
-	}
-	
-private:
-	TArray<FTexture2DRHIRef> TextureSet;
-	uint32 TextureIndex;
-};
-
 /**
  * SteamVR Head Mounted Display
  */
@@ -180,8 +135,6 @@ public:
 
 	virtual void UpdateScreenSettings(const FViewport* InViewport) override {}
 	
-	virtual bool AllocateRenderTargetTexture(uint32 Index, uint32 SizeX, uint32 SizeY, uint8 Format, uint32 NumMips, uint32 InTexFlags, uint32 InTargetableTextureFlags, FTexture2DRHIRef& OutTargetableTexture, FTexture2DRHIRef& OutShaderResourceTexture, uint32 NumSamples = 1) override;
-
 	virtual bool GetHMDDistortionEnabled(EShadingPath ShadingPath) const override;
 
 	virtual void OnBeginRendering_GameThread() override;
@@ -211,6 +164,9 @@ public:
 	}
 	virtual void CalculateRenderTargetSize(const class FViewport& Viewport, uint32& InOutSizeX, uint32& InOutSizeY) override;
 	virtual bool NeedReAllocateViewportRenderTarget(const class FViewport& Viewport) override;
+	virtual bool NeedReAllocateDepthTexture(const TRefCountPtr<struct IPooledRenderTarget>& DepthTarget) override;
+	virtual bool AllocateRenderTargetTexture(uint32 Index, uint32 SizeX, uint32 SizeY, uint8 Format, uint32 NumMips, uint32 InTexFlags, uint32 InTargetableTextureFlags, FTexture2DRHIRef& OutTargetableTexture, FTexture2DRHIRef& OutShaderResourceTexture, uint32 NumSamples = 1) override;
+	virtual bool AllocateDepthTexture(uint32 Index, uint32 SizeX, uint32 SizeY, uint8 Format, uint32 NumMips, uint32 Flags, uint32 TargetableTextureFlags, FTexture2DRHIRef& OutTargetableTexture, FTexture2DRHIRef& OutShaderResourceTexture, uint32 NumSamples = 1) override;
 
 	// IStereoLayers interface
 	// Create/Set/Get/Destroy inherited from TStereoLayerManager
@@ -256,28 +212,35 @@ public:
 			return bUseExplicitTimingMode;
 		}
 
+		FXRSwapChainPtr GetSwapChain() { return SwapChain; }
+		FXRSwapChainPtr GetDepthSwapChain() { return DepthSwapChain; }
+
 		/** Schedules BeginRendering_RHI on the RHI thread when in explicit timing mode */
 		void BeginRendering_RenderThread(FRHICommandListImmediate& RHICmdList);
 
 		/** Called only when we're in explicit timing mode, which needs to be paired with a call to PostPresentHandoff */
 		void BeginRendering_RHI();
 
+		void CreateSwapChain(const FTextureRHIRef& BindingTexture, TArray<FTextureRHIRef>&& SwapChainTextures);
+		void CreateDepthSwapChain(const FTextureRHIRef& BindingTexture, TArray<FTextureRHIRef>&& SwapChainTextures);
 
 		// Virtual interface implemented by subclasses
 		virtual void Reset() = 0;
+
 	private:
 		virtual void FinishRendering() = 0;
 
 	protected:
 		
 		FSteamVRHMD*			Plugin;
+		FXRSwapChainPtr			SwapChain;
+		FXRSwapChainPtr			DepthSwapChain;
+
 		bool					bInitialized;
 		
 		/** If we use explicit timing mode, we must have matching calls to BeginRendering_RHI and PostPresentHandoff */
 		bool					bUseExplicitTimingMode;
 		bool NeedsPostPresentHandoff() const;
-		
-
 	};
 
 #if PLATFORM_WINDOWS
@@ -286,13 +249,9 @@ public:
 	public:
 		D3D11Bridge(FSteamVRHMD* plugin);
 
-
 		virtual void FinishRendering() override;
 		virtual void UpdateViewport(const FViewport& Viewport, FRHIViewport* InViewportRHI) override;
 		virtual void Reset() override;
-
-	protected:
-		ID3D11Texture2D* RenderTargetTexture = NULL;
 	};
 #endif // PLATFORM_WINDOWS
 
@@ -319,10 +278,6 @@ public:
 		virtual void FinishRendering() override;
 		virtual void UpdateViewport(const FViewport& Viewport, FRHIViewport* InViewportRHI) override;
 		virtual void Reset() override;
-
-	protected:
-
-		FTexture2DRHIRef RenderTargetTexture;
 	};
 
 	class OpenGLBridge : public BridgeBaseImpl
@@ -333,10 +288,6 @@ public:
 		virtual void FinishRendering() override;
 		virtual void UpdateViewport(const FViewport& Viewport, FRHIViewport* InViewportRHI) override;
 		virtual void Reset() override;
-
-	protected:
-		GLuint RenderTargetTexture = 0;
-
 	};
 	
 #elif PLATFORM_MAC
@@ -350,8 +301,6 @@ public:
 		virtual void Reset() override;
 		
 		IOSurfaceRef GetSurface(const uint32 SizeX, const uint32 SizeY);
-		
-		FTexture2DRHIRef TextureSet;
 	};
 #endif // PLATFORM_MAC
 
@@ -442,6 +391,36 @@ public:
 
 		return out;
 	}
+
+	static FORCEINLINE vr::HmdMatrix44_t ToHmdMatrix44(const FMatrix& tm)
+	{
+		// Rows and columns are swapped between vr::HmdMatrix44_t and FMatrix
+		vr::HmdMatrix44_t out;
+
+		out.m[0][0] = tm.M[0][0];
+		out.m[1][0] = tm.M[0][1];
+		out.m[2][0] = tm.M[0][2];
+		out.m[3][0] = tm.M[0][3];
+
+		out.m[0][1] = tm.M[1][0];
+		out.m[1][1] = tm.M[1][1];
+		out.m[2][1] = tm.M[1][2];
+		out.m[3][1] = tm.M[1][3];
+
+		out.m[0][2] = tm.M[2][0];
+		out.m[1][2] = tm.M[2][1];
+		out.m[2][2] = tm.M[2][2];
+		out.m[3][2] = tm.M[2][3];
+
+		out.m[0][3] = tm.M[3][0];
+		out.m[1][3] = tm.M[3][1];
+		out.m[2][3] = tm.M[3][2];
+		out.m[3][3] = tm.M[3][3];
+
+		return out;
+	}
+
+
 private:
 
 	void SetupOcclusionMeshes();
diff --git a/Engine/Plugins/Runtime/Steam/SteamVR/Source/SteamVR/Private/SteamVRRender.cpp b/Engine/Plugins/Runtime/Steam/SteamVR/Source/SteamVR/Private/SteamVRRender.cpp
--- a/Engine/Plugins/Runtime/Steam/SteamVR/Source/SteamVR/Private/SteamVRRender.cpp
+++ b/Engine/Plugins/Runtime/Steam/SteamVR/Source/SteamVR/Private/SteamVRRender.cpp
@@ -112,6 +112,22 @@ void FSteamVRHMD::BridgeBaseImpl::BeginRendering_RHI()
 	Plugin->VRCompositor->SubmitExplicitTimingData();
 }
 
+void FSteamVRHMD::BridgeBaseImpl::CreateSwapChain(const FTextureRHIRef& BindingTexture, TArray<FTextureRHIRef>&& SwapChainTextures)
+{
+	check(IsInRenderingThread());
+	check(SwapChainTextures.Num());
+
+	SwapChain = CreateXRSwapChain(MoveTemp(SwapChainTextures), BindingTexture);
+}
+
+void FSteamVRHMD::BridgeBaseImpl::CreateDepthSwapChain(const FTextureRHIRef& BindingTexture, TArray<FTextureRHIRef>&& SwapChainTextures)
+{
+	check(IsInRenderingThread());
+	check(SwapChainTextures.Num());
+
+	DepthSwapChain = CreateXRSwapChain(MoveTemp(SwapChainTextures), BindingTexture);
+}
+
 bool FSteamVRHMD::BridgeBaseImpl::Present(int& SyncInterval)
 {
 	check(IsRunningRHIInSeparateThread() ? IsInRHIThread() : IsInRenderingThread());
@@ -123,6 +139,10 @@ bool FSteamVRHMD::BridgeBaseImpl::Present(int& SyncInterval)
 
 	FinishRendering();
 
+	// Increment swap chain index post-swap.
+	SwapChain->IncrementSwapChainIndex_RHIThread();
+	DepthSwapChain->IncrementSwapChainIndex_RHIThread();
+
 	return true;
 }
 
@@ -153,17 +173,22 @@ void FSteamVRHMD::BridgeBaseImpl::PostPresent()
 #if PLATFORM_WINDOWS
 
 FSteamVRHMD::D3D11Bridge::D3D11Bridge(FSteamVRHMD* plugin):
-	BridgeBaseImpl(plugin),
-	RenderTargetTexture(nullptr)
+	BridgeBaseImpl(plugin)
 {
+
 }
 
 void FSteamVRHMD::D3D11Bridge::FinishRendering()
 {
-	vr::Texture_t Texture;
-	Texture.handle = RenderTargetTexture;
+	vr::VRTextureWithDepth_t Texture;
+	Texture.handle = SwapChain->GetTexture2D()->GetNativeResource();
 	Texture.eType = vr::TextureType_DirectX;
 	Texture.eColorSpace = vr::ColorSpace_Auto;
+	Texture.depth.handle = DepthSwapChain->GetTexture2D()->GetNativeResource();
+
+	// Set the texture depth range to follow Unreal's inverted-Z depth settings (near is 1.0f, far is 0.0f).
+	Texture.depth.vRange.v[0] = 1.0f;
+	Texture.depth.vRange.v[1] = 0.0f;
 
 	vr::VRTextureBounds_t LeftBounds;
 	LeftBounds.uMin = 0.0f;
@@ -171,7 +196,12 @@ void FSteamVRHMD::D3D11Bridge::FinishRendering()
 	LeftBounds.vMin = 0.0f;
 	LeftBounds.vMax = 1.0f;
 	
-	vr::EVRCompositorError Error = Plugin->VRCompositor->Submit(vr::Eye_Left, &Texture, &LeftBounds);
+	Texture.depth.mProjection = ToHmdMatrix44(Plugin->GetStereoProjectionMatrix(eSSP_LEFT_EYE));
+	// Negate projection[3][2] to signal reverse-Z to the SteamVR runtime
+	Texture.depth.mProjection.m[3][2] = -Texture.depth.mProjection.m[3][2];
+	// Rescale the projection (our projection value is 10.0f here, since our units are cm, and SteamVR works in meters).
+	Texture.depth.mProjection.m[2][3] *= 0.01f;
+	vr::EVRCompositorError Error = Plugin->VRCompositor->Submit(vr::Eye_Left, &Texture, &LeftBounds, vr::EVRSubmitFlags::Submit_TextureWithDepth);
 
 	vr::VRTextureBounds_t RightBounds;
 	RightBounds.uMin = 0.5f;
@@ -179,13 +209,17 @@ void FSteamVRHMD::D3D11Bridge::FinishRendering()
 	RightBounds.vMin = 0.0f;
 	RightBounds.vMax = 1.0f;
 	   
-	Texture.handle = RenderTargetTexture;
-	Error = Plugin->VRCompositor->Submit(vr::Eye_Right, &Texture, &RightBounds);
+	Texture.depth.mProjection = ToHmdMatrix44(Plugin->GetStereoProjectionMatrix(eSSP_RIGHT_EYE));
+	// Negate projection[3][2] to signal reverse-Z to the SteamVR runtime
+	Texture.depth.mProjection.m[3][2] = -Texture.depth.mProjection.m[3][2];
+	// Rescale the projection (our projection value is 10.0f here, since our units are cm, and SteamVR works in meters).
+	Texture.depth.mProjection.m[2][3] *= 0.01f;
+	Error = Plugin->VRCompositor->Submit(vr::Eye_Right, &Texture, &RightBounds, vr::EVRSubmitFlags::Submit_TextureWithDepth);
 
 	static bool FirstError = true;
 	if (FirstError && Error != vr::VRCompositorError_None)
 	{
-		UE_LOG(LogHMD, Log, TEXT("Warning:  SteamVR Compositor had an error on present (%d)"), (int32)Error);
+		UE_LOG(LogHMD, Log, TEXT("Warning: SteamVR Compositor had an error on present (%d)"), (int32)Error);
 		FirstError = false;
 	}
 }
@@ -201,22 +235,14 @@ void FSteamVRHMD::D3D11Bridge::UpdateViewport(const FViewport& Viewport, FRHIVie
 
 	const FTexture2DRHIRef& RT = Viewport.GetRenderTargetTexture();
 	check(IsValidRef(RT));
-
-	if (RenderTargetTexture != nullptr)
-	{
-		RenderTargetTexture->Release();	//@todo steamvr: need to also release in reset
-	}
-
-	RenderTargetTexture = (ID3D11Texture2D*)RT->GetNativeResource();
-	RenderTargetTexture->AddRef();
+	check(RT->GetTexture2D() == SwapChain->GetTexture2D());
 }
 
 #endif // PLATFORM_WINDOWS
 
 #if !PLATFORM_MAC
 FSteamVRHMD::VulkanBridge::VulkanBridge(FSteamVRHMD* plugin):
-	BridgeBaseImpl(plugin),
-	RenderTargetTexture(0)
+	BridgeBaseImpl(plugin)
 {
 	bInitialized = true;
 	bUseExplicitTimingMode = true;
@@ -225,20 +251,34 @@ FSteamVRHMD::VulkanBridge::VulkanBridge(FSteamVRHMD* plugin):
 void FSteamVRHMD::VulkanBridge::FinishRendering()
 {
 	auto vlkRHI = static_cast<FVulkanDynamicRHI*>(GDynamicRHI);
-
-	if(RenderTargetTexture.IsValid())
+	if(SwapChain->GetTexture2D())
 	{
-		FVulkanTexture2D* Texture2D = (FVulkanTexture2D*)RenderTargetTexture.GetReference();
+		FVulkanTexture2D* Texture2D = (FVulkanTexture2D*)SwapChain->GetTexture2D();
+		FVulkanTexture2D* DepthTexture2D = (FVulkanTexture2D*)DepthSwapChain->GetTexture2D();
+
 		FVulkanCommandListContext& ImmediateContext = vlkRHI->GetDevice()->GetImmediateContext();
+
+		// Color layout
 		VkImageLayout& CurrentLayout = ImmediateContext.GetTransitionAndLayoutManager().FindOrAddLayoutRW(Texture2D->Surface.Image, VK_IMAGE_LAYOUT_UNDEFINED);
 		bool bHadLayout = (CurrentLayout != VK_IMAGE_LAYOUT_UNDEFINED);
+
+		VkImageLayout& CurrentDepthLayout = ImmediateContext.GetTransitionAndLayoutManager().FindOrAddLayoutRW(DepthTexture2D->Surface.Image, VK_IMAGE_LAYOUT_UNDEFINED);
+		bool bDepthHadLayout = (CurrentLayout != VK_IMAGE_LAYOUT_UNDEFINED);
+
 		FVulkanCmdBuffer* CmdBuffer = ImmediateContext.GetCommandBufferManager()->GetUploadCmdBuffer();
 		VkImageSubresourceRange SubresourceRange = { VK_IMAGE_ASPECT_COLOR_BIT, 0, 1, 0, 1 };
+		VkImageSubresourceRange SubresourceRangeDepth = { VK_IMAGE_ASPECT_DEPTH_BIT | VK_IMAGE_ASPECT_STENCIL_BIT, 0, 1, 0, 1 };
+
 		if (CurrentLayout != VK_IMAGE_LAYOUT_TRANSFER_SRC_OPTIMAL)
 		{
 			vlkRHI->VulkanSetImageLayout(CmdBuffer->GetHandle(), Texture2D->Surface.Image, CurrentLayout, VK_IMAGE_LAYOUT_TRANSFER_SRC_OPTIMAL, SubresourceRange);
 		}
 
+		if (CurrentDepthLayout != VK_IMAGE_LAYOUT_TRANSFER_SRC_OPTIMAL)
+		{
+			vlkRHI->VulkanSetImageLayout(CmdBuffer->GetHandle(), DepthTexture2D->Surface.Image, CurrentLayout, VK_IMAGE_LAYOUT_TRANSFER_SRC_OPTIMAL, SubresourceRangeDepth);
+		}
+
 		vr::VRTextureBounds_t LeftBounds;
 		LeftBounds.uMin = 0.0f;
 		LeftBounds.uMax = 0.5f;
@@ -251,22 +291,43 @@ void FSteamVRHMD::VulkanBridge::FinishRendering()
 		RightBounds.vMin = 0.0f;
 		RightBounds.vMax = 1.0f;
 
-		vr::VRVulkanTextureData_t vulkanData {};
-		vulkanData.m_pInstance			= vlkRHI->GetInstance();
-		vulkanData.m_pDevice			= vlkRHI->GetDevice()->GetInstanceHandle();
-		vulkanData.m_pPhysicalDevice	= vlkRHI->GetDevice()->GetPhysicalHandle();
-		vulkanData.m_pQueue				= vlkRHI->GetDevice()->GetGraphicsQueue()->GetHandle();
-		vulkanData.m_nQueueFamilyIndex	= vlkRHI->GetDevice()->GetGraphicsQueue()->GetFamilyIndex();
-		vulkanData.m_nImage				= (uint64_t)Texture2D->Surface.Image;
-		vulkanData.m_nWidth				= Texture2D->Surface.Width;
-		vulkanData.m_nHeight			= Texture2D->Surface.Height;
-		vulkanData.m_nFormat			= (uint32_t)Texture2D->Surface.ViewFormat;
-		vulkanData.m_nSampleCount = 1;
-
-		vr::Texture_t texture = {&vulkanData, vr::TextureType_Vulkan, vr::ColorSpace_Auto};
-
-		Plugin->VRCompositor->Submit(vr::Eye_Left, &texture, &LeftBounds);
-		Plugin->VRCompositor->Submit(vr::Eye_Right, &texture, &RightBounds);
+		vr::VRVulkanTextureData_t VulkanTextureDataColor {};
+		VulkanTextureDataColor.m_pInstance			= vlkRHI->GetInstance();
+		VulkanTextureDataColor.m_pDevice			= vlkRHI->GetDevice()->GetInstanceHandle();
+		VulkanTextureDataColor.m_pPhysicalDevice	= vlkRHI->GetDevice()->GetPhysicalHandle();
+		VulkanTextureDataColor.m_pQueue				= vlkRHI->GetDevice()->GetGraphicsQueue()->GetHandle();
+		VulkanTextureDataColor.m_nQueueFamilyIndex	= vlkRHI->GetDevice()->GetGraphicsQueue()->GetFamilyIndex();
+		VulkanTextureDataColor.m_nImage				= (uint64_t)Texture2D->Surface.Image;
+		VulkanTextureDataColor.m_nWidth				= Texture2D->Surface.Width;
+		VulkanTextureDataColor.m_nHeight			= Texture2D->Surface.Height;
+		VulkanTextureDataColor.m_nFormat			= (uint32_t)Texture2D->Surface.ViewFormat;
+		VulkanTextureDataColor.m_nSampleCount = 1;
+
+		vr::VRVulkanTextureData_t VulkanTextureDataDepth{};
+		VulkanTextureDataDepth.m_pInstance = vlkRHI->GetInstance();
+		VulkanTextureDataDepth.m_pDevice = vlkRHI->GetDevice()->GetInstanceHandle();
+		VulkanTextureDataDepth.m_pPhysicalDevice = vlkRHI->GetDevice()->GetPhysicalHandle();
+		VulkanTextureDataDepth.m_pQueue = vlkRHI->GetDevice()->GetGraphicsQueue()->GetHandle();
+		VulkanTextureDataDepth.m_nQueueFamilyIndex = vlkRHI->GetDevice()->GetGraphicsQueue()->GetFamilyIndex();
+		VulkanTextureDataDepth.m_nImage = (uint64_t)DepthTexture2D->Surface.Image;
+		VulkanTextureDataDepth.m_nWidth = DepthTexture2D->Surface.Width;
+		VulkanTextureDataDepth.m_nHeight = DepthTexture2D->Surface.Height;
+		VulkanTextureDataDepth.m_nFormat = (uint32_t)DepthTexture2D->Surface.ViewFormat;
+		VulkanTextureDataDepth.m_nSampleCount = 1;
+
+		vr::VRTextureWithDepth_t Texture;
+		Texture.handle = &VulkanTextureDataColor;
+		Texture.eColorSpace = vr::ColorSpace_Auto;
+		Texture.eType = vr::TextureType_Vulkan;
+		Texture.depth.handle = &VulkanTextureDataDepth;
+		Texture.depth.vRange.v[0] = 1.0f;
+		Texture.depth.vRange.v[1] = 0.0f;
+
+		Texture.depth.mProjection = ToHmdMatrix44(Plugin->GetStereoProjectionMatrix(eSSP_LEFT_EYE));
+		Plugin->VRCompositor->Submit(vr::Eye_Left, &Texture, &LeftBounds, vr::EVRSubmitFlags::Submit_TextureWithDepth);
+
+		Texture.depth.mProjection = ToHmdMatrix44(Plugin->GetStereoProjectionMatrix(eSSP_RIGHT_EYE));
+		Plugin->VRCompositor->Submit(vr::Eye_Right, &Texture, &RightBounds, vr::EVRSubmitFlags::Submit_TextureWithDepth);
 
 		if (bHadLayout && CurrentLayout != VK_IMAGE_LAYOUT_TRANSFER_SRC_OPTIMAL)
 		{
@@ -277,34 +338,42 @@ void FSteamVRHMD::VulkanBridge::FinishRendering()
 			CurrentLayout = VK_IMAGE_LAYOUT_TRANSFER_SRC_OPTIMAL;
 		}
 
+		if (bDepthHadLayout && CurrentDepthLayout != VK_IMAGE_LAYOUT_TRANSFER_SRC_OPTIMAL)
+		{
+			vlkRHI->VulkanSetImageLayout(CmdBuffer->GetHandle(), DepthTexture2D->Surface.Image, VK_IMAGE_LAYOUT_TRANSFER_SRC_OPTIMAL, CurrentDepthLayout, SubresourceRangeDepth);
+		}
+		else
+		{
+			CurrentDepthLayout = VK_IMAGE_LAYOUT_TRANSFER_SRC_OPTIMAL;
+		}
+
 		ImmediateContext.GetCommandBufferManager()->SubmitUploadCmdBuffer();
 	}
 }
 
 void FSteamVRHMD::VulkanBridge::Reset()
 {
-
 }
 
 void FSteamVRHMD::VulkanBridge::UpdateViewport(const FViewport& Viewport, FRHIViewport* InViewportRHI)
 {
-	RenderTargetTexture = Viewport.GetRenderTargetTexture();
-	check(IsValidRef(RenderTargetTexture));
 }
 
 FSteamVRHMD::OpenGLBridge::OpenGLBridge(FSteamVRHMD* plugin):
-	BridgeBaseImpl(plugin),
-	RenderTargetTexture(0)
+	BridgeBaseImpl(plugin)
 {
 	bInitialized = true;
 }
 
 void FSteamVRHMD::OpenGLBridge::FinishRendering()
 {
+	GLuint RenderTargetTexture = *reinterpret_cast<GLuint*>(SwapChain->GetTexture2D()->GetNativeResource());
+	GLuint DepthTargetTexture = *reinterpret_cast<GLuint*>(DepthSwapChain->GetTexture2D()->GetNativeResource());
+
 	// Yaakuro:
 	// TODO This is a workaround. After exiting VR Editor the texture gets invalid at some point.
 	// Need to find it. This at least prevents to use this method when texture name is not valid anymore.
-	if(!glIsTexture(RenderTargetTexture))
+	if (!glIsTexture(RenderTargetTexture) || !glIsTexture(DepthTargetTexture))
 	{
 		return;
 	}
@@ -321,18 +390,23 @@ void FSteamVRHMD::OpenGLBridge::FinishRendering()
 	RightBounds.vMin = 1.0f;
 	RightBounds.vMax = 0.0f;
 
-	vr::Texture_t Texture;
+	vr::VRTextureWithDepth_t Texture;
 	Texture.handle = reinterpret_cast<void*>(static_cast<size_t>(RenderTargetTexture));
 	Texture.eType = vr::TextureType_OpenGL;
 	Texture.eColorSpace = vr::ColorSpace_Auto;
+	Texture.depth.handle = reinterpret_cast<void*>(static_cast<size_t>(DepthTargetTexture));
+	Texture.depth.vRange.v[0] = 1.0f;
+	Texture.depth.vRange.v[1] = 0.0f;
+
+	Texture.depth.mProjection = ToHmdMatrix44(Plugin->GetStereoProjectionMatrix(eSSP_LEFT_EYE));
+	Plugin->VRCompositor->Submit(vr::Eye_Left, &Texture, &LeftBounds, vr::EVRSubmitFlags::Submit_TextureWithDepth);
 
-	Plugin->VRCompositor->Submit(vr::Eye_Left, &Texture, &LeftBounds);
-	Plugin->VRCompositor->Submit(vr::Eye_Right, &Texture, &RightBounds);
+	Texture.depth.mProjection = ToHmdMatrix44(Plugin->GetStereoProjectionMatrix(eSSP_RIGHT_EYE));
+	Plugin->VRCompositor->Submit(vr::Eye_Right, &Texture, &RightBounds, vr::EVRSubmitFlags::Submit_TextureWithDepth);
 }
 
 void FSteamVRHMD::OpenGLBridge::Reset()
 {
-	RenderTargetTexture = 0;
 }
 
 void FSteamVRHMD::OpenGLBridge::UpdateViewport(const FViewport& Viewport, FRHIViewport* InViewportRHI)
@@ -341,8 +415,7 @@ void FSteamVRHMD::OpenGLBridge::UpdateViewport(const FViewport& Viewport, FRHIVi
 
 	const FTexture2DRHIRef& RT = Viewport.GetRenderTargetTexture();
 	check(IsValidRef(RT));
-
-	RenderTargetTexture = *reinterpret_cast<GLuint*>(RT->GetNativeResource());
+	check(RT == SwapChain->GetTexture2D());
 }
 
 #elif PLATFORM_MAC
@@ -353,20 +426,20 @@ FSteamVRHMD::MetalBridge::MetalBridge(FSteamVRHMD* plugin):
 
 void FSteamVRHMD::MetalBridge::FinishRendering()
 {
-	check(TextureSet.IsValid());
-
 	vr::VRTextureBounds_t LeftBounds;
 	LeftBounds.uMin = 0.0f;
 	LeftBounds.uMax = 0.5f;
 	LeftBounds.vMin = 0.0f;
 	LeftBounds.vMax = 1.0f;
 	
-	id<MTLTexture> TextureHandle = (id<MTLTexture>)TextureSet->GetNativeResource();
-	
+	id<MTLTexture> TextureHandle = (id<MTLTexture>)SwapChain->GetTexture2D()->GetNativeResource();
+
+	// @todo: Add depth.
 	vr::Texture_t Texture;
 	Texture.handle = (void*)TextureHandle.iosurface;
 	Texture.eType = vr::TextureType_IOSurface;
 	Texture.eColorSpace = vr::ColorSpace_Auto;
+
 	vr::EVRCompositorError Error = Plugin->VRCompositor->Submit(vr::Eye_Left, &Texture, &LeftBounds);
 
 	vr::VRTextureBounds_t RightBounds;
@@ -383,8 +456,6 @@ void FSteamVRHMD::MetalBridge::FinishRendering()
 		UE_LOG(LogHMD, Log, TEXT("Warning:  SteamVR Compositor had an error on present (%d)"), (int32)Error);
 		FirstError = false;
 	}
-
-	static_cast<FRHITextureSet2D*>(TextureSet.GetReference())->Advance();
 }
 
 void FSteamVRHMD::MetalBridge::Reset()
@@ -394,6 +465,7 @@ void FSteamVRHMD::MetalBridge::Reset()
 IOSurfaceRef FSteamVRHMD::MetalBridge::GetSurface(const uint32 SizeX, const uint32 SizeY)
 {
 	// @todo: Get our man in MacVR to switch to a modern & secure method of IOSurface sharing...
+	// @todo: Also add support for depth.
 PRAGMA_DISABLE_DEPRECATION_WARNINGS
 	const NSDictionary* SurfaceDefinition = @{
 											(id)kIOSurfaceWidth: @(SizeX),
