Fix inversion bug with steamVR depth passthrough. Turns out we don't need to negate any elements in the projection matrix to get correct depth results on an Oculus headset via SteamVR (and doing so broke rendering on the Index).

diff --git a/Engine/Plugins/Runtime/Steam/SteamVR/Source/SteamVR/Private/SteamVRRender.cpp b/Engine/Plugins/Runtime/Steam/SteamVR/Source/SteamVR/Private/SteamVRRender.cpp
--- a/Engine/Plugins/Runtime/Steam/SteamVR/Source/SteamVR/Private/SteamVRRender.cpp
+++ b/Engine/Plugins/Runtime/Steam/SteamVR/Source/SteamVR/Private/SteamVRRender.cpp
@@ -178,43 +178,58 @@ FSteamVRHMD::D3D11Bridge::D3D11Bridge(FSteamVRHMD* plugin):
 
 }
 
+static TAutoConsoleVariable<int> CVarEnableDepthSubmission(
+	TEXT("vr.EnableSteamVRDepthSubmission"),
+	1,
+	TEXT("By default, depth is passed through in SteamVR for devices that support depth. Set this flag to 0 to disable depth submission."),
+	ECVF_Default);
+
 void FSteamVRHMD::D3D11Bridge::FinishRendering()
 {
+	bool bSubmitDepth = CVarEnableDepthSubmission->GetInt() > 0;
+	vr::EVRSubmitFlags Flags = bSubmitDepth ? vr::EVRSubmitFlags::Submit_TextureWithDepth : vr::EVRSubmitFlags::Submit_Default;
+
 	vr::VRTextureWithDepth_t Texture;
 	Texture.handle = SwapChain->GetTexture2D()->GetNativeResource();
 	Texture.eType = vr::TextureType_DirectX;
 	Texture.eColorSpace = vr::ColorSpace_Auto;
-	Texture.depth.handle = DepthSwapChain->GetTexture2D()->GetNativeResource();
 
-	// Set the texture depth range to follow Unreal's inverted-Z depth settings (near is 1.0f, far is 0.0f).
-	Texture.depth.vRange.v[0] = 1.0f;
-	Texture.depth.vRange.v[1] = 0.0f;
+	if (bSubmitDepth)
+	{
+		// If this flag is false, the struct will be treated as a vr::Texture_t and these entries will be ignored - so we can skip this if not submitting depth
+		Texture.depth.handle = DepthSwapChain->GetTexture2D()->GetNativeResource();
+
+		// Set the texture depth range to follow Unreal's inverted-Z depth settings (near is 1.0f, far is 0.0f).
+		Texture.depth.vRange.v[0] = 1.0f;
+		Texture.depth.vRange.v[1] = 0.0f;
+
+		Texture.depth.mProjection = ToHmdMatrix44(Plugin->GetStereoProjectionMatrix(eSSP_LEFT_EYE));
+		// Rescale the projection (our projection value is 10.0f here, since our units are cm, and SteamVR works in meters).
+		Texture.depth.mProjection.m[2][3] *= 0.01f;
+	}
 
 	vr::VRTextureBounds_t LeftBounds;
 	LeftBounds.uMin = 0.0f;
 	LeftBounds.uMax = 0.5f;
 	LeftBounds.vMin = 0.0f;
 	LeftBounds.vMax = 1.0f;
-	
-	Texture.depth.mProjection = ToHmdMatrix44(Plugin->GetStereoProjectionMatrix(eSSP_LEFT_EYE));
-	// Negate projection[3][2] to signal reverse-Z to the SteamVR runtime
-	Texture.depth.mProjection.m[3][2] = -Texture.depth.mProjection.m[3][2];
-	// Rescale the projection (our projection value is 10.0f here, since our units are cm, and SteamVR works in meters).
-	Texture.depth.mProjection.m[2][3] *= 0.01f;
-	vr::EVRCompositorError Error = Plugin->VRCompositor->Submit(vr::Eye_Left, &Texture, &LeftBounds, vr::EVRSubmitFlags::Submit_TextureWithDepth);
+
+	vr::EVRCompositorError Error = Plugin->VRCompositor->Submit(vr::Eye_Left, &Texture, &LeftBounds, Flags);
 
 	vr::VRTextureBounds_t RightBounds;
 	RightBounds.uMin = 0.5f;
 	RightBounds.uMax = 1.0f;
 	RightBounds.vMin = 0.0f;
 	RightBounds.vMax = 1.0f;
-	   
-	Texture.depth.mProjection = ToHmdMatrix44(Plugin->GetStereoProjectionMatrix(eSSP_RIGHT_EYE));
-	// Negate projection[3][2] to signal reverse-Z to the SteamVR runtime
-	Texture.depth.mProjection.m[3][2] = -Texture.depth.mProjection.m[3][2];
-	// Rescale the projection (our projection value is 10.0f here, since our units are cm, and SteamVR works in meters).
-	Texture.depth.mProjection.m[2][3] *= 0.01f;
-	Error = Plugin->VRCompositor->Submit(vr::Eye_Right, &Texture, &RightBounds, vr::EVRSubmitFlags::Submit_TextureWithDepth);
+
+	if (bSubmitDepth)
+	{
+		Texture.depth.mProjection = ToHmdMatrix44(Plugin->GetStereoProjectionMatrix(eSSP_RIGHT_EYE));
+		// Rescale the projection (our projection value is 10.0f here, since our units are cm, and SteamVR works in meters).
+		Texture.depth.mProjection.m[2][3] *= 0.01f;
+	}
+
+	Error = Plugin->VRCompositor->Submit(vr::Eye_Right, &Texture, &RightBounds, Flags);
 
 	static bool FirstError = true;
 	if (FirstError && Error != vr::VRCompositorError_None)
@@ -250,6 +265,9 @@ FSteamVRHMD::VulkanBridge::VulkanBridge(FSteamVRHMD* plugin):
 
 void FSteamVRHMD::VulkanBridge::FinishRendering()
 {
+	bool bSubmitDepth = CVarEnableDepthSubmission->GetInt() > 0;
+	vr::EVRSubmitFlags Flags = bSubmitDepth ? vr::EVRSubmitFlags::Submit_TextureWithDepth : vr::EVRSubmitFlags::Submit_Default;
+
 	auto vlkRHI = static_cast<FVulkanDynamicRHI*>(GDynamicRHI);
 	if(SwapChain->GetTexture2D())
 	{
@@ -261,24 +279,15 @@ void FSteamVRHMD::VulkanBridge::FinishRendering()
 		// Color layout
 		VkImageLayout& CurrentLayout = ImmediateContext.GetTransitionAndLayoutManager().FindOrAddLayoutRW(Texture2D->Surface.Image, VK_IMAGE_LAYOUT_UNDEFINED);
 		bool bHadLayout = (CurrentLayout != VK_IMAGE_LAYOUT_UNDEFINED);
-
-		VkImageLayout& CurrentDepthLayout = ImmediateContext.GetTransitionAndLayoutManager().FindOrAddLayoutRW(DepthTexture2D->Surface.Image, VK_IMAGE_LAYOUT_UNDEFINED);
-		bool bDepthHadLayout = (CurrentLayout != VK_IMAGE_LAYOUT_UNDEFINED);
-
+		
 		FVulkanCmdBuffer* CmdBuffer = ImmediateContext.GetCommandBufferManager()->GetUploadCmdBuffer();
 		VkImageSubresourceRange SubresourceRange = { VK_IMAGE_ASPECT_COLOR_BIT, 0, 1, 0, 1 };
-		VkImageSubresourceRange SubresourceRangeDepth = { VK_IMAGE_ASPECT_DEPTH_BIT | VK_IMAGE_ASPECT_STENCIL_BIT, 0, 1, 0, 1 };
 
 		if (CurrentLayout != VK_IMAGE_LAYOUT_TRANSFER_SRC_OPTIMAL)
 		{
 			vlkRHI->VulkanSetImageLayout(CmdBuffer->GetHandle(), Texture2D->Surface.Image, CurrentLayout, VK_IMAGE_LAYOUT_TRANSFER_SRC_OPTIMAL, SubresourceRange);
 		}
 
-		if (CurrentDepthLayout != VK_IMAGE_LAYOUT_TRANSFER_SRC_OPTIMAL)
-		{
-			vlkRHI->VulkanSetImageLayout(CmdBuffer->GetHandle(), DepthTexture2D->Surface.Image, CurrentLayout, VK_IMAGE_LAYOUT_TRANSFER_SRC_OPTIMAL, SubresourceRangeDepth);
-		}
-
 		vr::VRTextureBounds_t LeftBounds;
 		LeftBounds.uMin = 0.0f;
 		LeftBounds.uMax = 0.5f;
@@ -303,48 +312,78 @@ void FSteamVRHMD::VulkanBridge::FinishRendering()
 		VulkanTextureDataColor.m_nFormat			= (uint32_t)Texture2D->Surface.ViewFormat;
 		VulkanTextureDataColor.m_nSampleCount = 1;
 
-		vr::VRVulkanTextureData_t VulkanTextureDataDepth{};
-		VulkanTextureDataDepth.m_pInstance = vlkRHI->GetInstance();
-		VulkanTextureDataDepth.m_pDevice = vlkRHI->GetDevice()->GetInstanceHandle();
-		VulkanTextureDataDepth.m_pPhysicalDevice = vlkRHI->GetDevice()->GetPhysicalHandle();
-		VulkanTextureDataDepth.m_pQueue = vlkRHI->GetDevice()->GetGraphicsQueue()->GetHandle();
-		VulkanTextureDataDepth.m_nQueueFamilyIndex = vlkRHI->GetDevice()->GetGraphicsQueue()->GetFamilyIndex();
-		VulkanTextureDataDepth.m_nImage = (uint64_t)DepthTexture2D->Surface.Image;
-		VulkanTextureDataDepth.m_nWidth = DepthTexture2D->Surface.Width;
-		VulkanTextureDataDepth.m_nHeight = DepthTexture2D->Surface.Height;
-		VulkanTextureDataDepth.m_nFormat = (uint32_t)DepthTexture2D->Surface.ViewFormat;
-		VulkanTextureDataDepth.m_nSampleCount = 1;
-
-		vr::VRTextureWithDepth_t Texture;
-		Texture.handle = &VulkanTextureDataColor;
-		Texture.eColorSpace = vr::ColorSpace_Auto;
-		Texture.eType = vr::TextureType_Vulkan;
-		Texture.depth.handle = &VulkanTextureDataDepth;
-		Texture.depth.vRange.v[0] = 1.0f;
-		Texture.depth.vRange.v[1] = 0.0f;
-
-		Texture.depth.mProjection = ToHmdMatrix44(Plugin->GetStereoProjectionMatrix(eSSP_LEFT_EYE));
-		Plugin->VRCompositor->Submit(vr::Eye_Left, &Texture, &LeftBounds, vr::EVRSubmitFlags::Submit_TextureWithDepth);
-
-		Texture.depth.mProjection = ToHmdMatrix44(Plugin->GetStereoProjectionMatrix(eSSP_RIGHT_EYE));
-		Plugin->VRCompositor->Submit(vr::Eye_Right, &Texture, &RightBounds, vr::EVRSubmitFlags::Submit_TextureWithDepth);
-
-		if (bHadLayout && CurrentLayout != VK_IMAGE_LAYOUT_TRANSFER_SRC_OPTIMAL)
+		if (bSubmitDepth)
 		{
-			vlkRHI->VulkanSetImageLayout(CmdBuffer->GetHandle(), Texture2D->Surface.Image, VK_IMAGE_LAYOUT_TRANSFER_SRC_OPTIMAL, CurrentLayout, SubresourceRange);
+			VkImageSubresourceRange SubresourceRangeDepth = { VK_IMAGE_ASPECT_DEPTH_BIT | VK_IMAGE_ASPECT_STENCIL_BIT, 0, 1, 0, 1 };
+			VkImageLayout& CurrentDepthLayout = ImmediateContext.GetTransitionAndLayoutManager().FindOrAddLayoutRW(DepthTexture2D->Surface.Image, VK_IMAGE_LAYOUT_UNDEFINED);
+			bool bDepthHadLayout = (CurrentLayout != VK_IMAGE_LAYOUT_UNDEFINED);
+
+			if (CurrentDepthLayout != VK_IMAGE_LAYOUT_TRANSFER_SRC_OPTIMAL)
+			{
+				vlkRHI->VulkanSetImageLayout(CmdBuffer->GetHandle(), DepthTexture2D->Surface.Image, CurrentLayout, VK_IMAGE_LAYOUT_TRANSFER_SRC_OPTIMAL, SubresourceRangeDepth);
+			}
+
+			vr::VRVulkanTextureData_t VulkanTextureDataDepth{};
+			VulkanTextureDataDepth.m_pInstance = vlkRHI->GetInstance();
+			VulkanTextureDataDepth.m_pDevice = vlkRHI->GetDevice()->GetInstanceHandle();
+			VulkanTextureDataDepth.m_pPhysicalDevice = vlkRHI->GetDevice()->GetPhysicalHandle();
+			VulkanTextureDataDepth.m_pQueue = vlkRHI->GetDevice()->GetGraphicsQueue()->GetHandle();
+			VulkanTextureDataDepth.m_nQueueFamilyIndex = vlkRHI->GetDevice()->GetGraphicsQueue()->GetFamilyIndex();
+			VulkanTextureDataDepth.m_nImage = (uint64_t)DepthTexture2D->Surface.Image;
+			VulkanTextureDataDepth.m_nWidth = DepthTexture2D->Surface.Width;
+			VulkanTextureDataDepth.m_nHeight = DepthTexture2D->Surface.Height;
+			VulkanTextureDataDepth.m_nFormat = (uint32_t)DepthTexture2D->Surface.ViewFormat;
+			VulkanTextureDataDepth.m_nSampleCount = 1;
+
+			vr::VRTextureWithDepth_t Texture;
+			Texture.handle = &VulkanTextureDataColor;
+			Texture.eColorSpace = vr::ColorSpace_Auto;
+			Texture.eType = vr::TextureType_Vulkan;
+			Texture.depth.handle = &VulkanTextureDataDepth;
+			Texture.depth.vRange.v[0] = 1.0f;
+			Texture.depth.vRange.v[1] = 0.0f;
+
+			Texture.depth.mProjection = ToHmdMatrix44(Plugin->GetStereoProjectionMatrix(eSSP_LEFT_EYE));
+
+			// Rescale the projection (our projection value is 10.0f here, since our units are cm, and SteamVR works in meters).
+			Texture.depth.mProjection.m[2][3] *= 0.01f;
+
+			Plugin->VRCompositor->Submit(vr::Eye_Left, &Texture, &LeftBounds, vr::EVRSubmitFlags::Submit_TextureWithDepth);
+
+			Texture.depth.mProjection = ToHmdMatrix44(Plugin->GetStereoProjectionMatrix(eSSP_RIGHT_EYE));
+
+			// Rescale the projection (our projection value is 10.0f here, since our units are cm, and SteamVR works in meters).
+			Texture.depth.mProjection.m[2][3] *= 0.01f;
+
+			Plugin->VRCompositor->Submit(vr::Eye_Right, &Texture, &RightBounds, vr::EVRSubmitFlags::Submit_TextureWithDepth);
+
+			if (bDepthHadLayout && CurrentDepthLayout != VK_IMAGE_LAYOUT_TRANSFER_SRC_OPTIMAL)
+			{
+				vlkRHI->VulkanSetImageLayout(CmdBuffer->GetHandle(), DepthTexture2D->Surface.Image, VK_IMAGE_LAYOUT_TRANSFER_SRC_OPTIMAL, CurrentDepthLayout, SubresourceRangeDepth);
+			}
+			else
+			{
+				CurrentDepthLayout = VK_IMAGE_LAYOUT_TRANSFER_SRC_OPTIMAL;
+			}
 		}
 		else
 		{
-			CurrentLayout = VK_IMAGE_LAYOUT_TRANSFER_SRC_OPTIMAL;
+			vr::Texture_t Texture;
+			Texture.handle = &VulkanTextureDataColor;
+			Texture.eColorSpace = vr::ColorSpace_Auto;
+			Texture.eType = vr::TextureType_Vulkan;
+
+			Plugin->VRCompositor->Submit(vr::Eye_Left, &Texture, &LeftBounds, vr::EVRSubmitFlags::Submit_Default);
+			Plugin->VRCompositor->Submit(vr::Eye_Right, &Texture, &RightBounds, vr::EVRSubmitFlags::Submit_Default);
 		}
 
-		if (bDepthHadLayout && CurrentDepthLayout != VK_IMAGE_LAYOUT_TRANSFER_SRC_OPTIMAL)
+		if (bHadLayout && CurrentLayout != VK_IMAGE_LAYOUT_TRANSFER_SRC_OPTIMAL)
 		{
-			vlkRHI->VulkanSetImageLayout(CmdBuffer->GetHandle(), DepthTexture2D->Surface.Image, VK_IMAGE_LAYOUT_TRANSFER_SRC_OPTIMAL, CurrentDepthLayout, SubresourceRangeDepth);
+			vlkRHI->VulkanSetImageLayout(CmdBuffer->GetHandle(), Texture2D->Surface.Image, VK_IMAGE_LAYOUT_TRANSFER_SRC_OPTIMAL, CurrentLayout, SubresourceRange);
 		}
 		else
 		{
-			CurrentDepthLayout = VK_IMAGE_LAYOUT_TRANSFER_SRC_OPTIMAL;
+			CurrentLayout = VK_IMAGE_LAYOUT_TRANSFER_SRC_OPTIMAL;
 		}
 
 		ImmediateContext.GetCommandBufferManager()->SubmitUploadCmdBuffer();
