Fix linux build-break (incorrect _TEXT macro usage and CVar defined inside #if PLATFORM_WINDOWS block).

diff --git a/Engine/Plugins/Runtime/Steam/SteamVR/Source/SteamVR/Private/SteamVRHMD.cpp b/Engine/Plugins/Runtime/Steam/SteamVR/Source/SteamVR/Private/SteamVRHMD.cpp
--- a/Engine/Plugins/Runtime/Steam/SteamVR/Source/SteamVR/Private/SteamVRHMD.cpp
+++ b/Engine/Plugins/Runtime/Steam/SteamVR/Source/SteamVR/Private/SteamVRHMD.cpp
@@ -1626,7 +1626,7 @@ bool FSteamVRHMD::AllocateDepthTexture(uint32 Index, uint32 SizeX, uint32 SizeY,
 	FClearValueBinding ClearValue(0.0f, 0);
 	ClearValue.ColorBinding = EClearBinding::EDepthStencilBound;
 	FRHIResourceCreateInfo CreateInfo(ClearValue);
-	CreateInfo.DebugName = _TEXT("SteamVRDepthStencil");
+	CreateInfo.DebugName = TEXT("SteamVRDepthStencil");
 
 	for (uint32 SwapChainIter = 0; SwapChainIter < SteamVRSwapChainLength; ++SwapChainIter)
 	{
