Fix some typo errors in depth submission for Vulkan. Temporarily disable depth submission when Vulkan RHI used.

diff --git a/Engine/Plugins/Runtime/Steam/SteamVR/Source/SteamVR/Private/SteamVRRender.cpp b/Engine/Plugins/Runtime/Steam/SteamVR/Source/SteamVR/Private/SteamVRRender.cpp
--- a/Engine/Plugins/Runtime/Steam/SteamVR/Source/SteamVR/Private/SteamVRRender.cpp
+++ b/Engine/Plugins/Runtime/Steam/SteamVR/Source/SteamVR/Private/SteamVRRender.cpp
@@ -265,7 +265,8 @@ FSteamVRHMD::VulkanBridge::VulkanBridge(FSteamVRHMD* plugin):
 
 void FSteamVRHMD::VulkanBridge::FinishRendering()
 {
-	bool bSubmitDepth = CVarEnableDepthSubmission->GetInt() > 0;
+	// Disable depth-submission until Vulkan "device lost" error on submission is tracked down.
+	bool bSubmitDepth = false; // CVarEnableDepthSubmission->GetInt() > 0;
 	vr::EVRSubmitFlags Flags = bSubmitDepth ? vr::EVRSubmitFlags::Submit_TextureWithDepth : vr::EVRSubmitFlags::Submit_Default;
 
 	auto vlkRHI = static_cast<FVulkanDynamicRHI*>(GDynamicRHI);
@@ -316,11 +317,11 @@ void FSteamVRHMD::VulkanBridge::FinishRendering()
 		{
 			VkImageSubresourceRange SubresourceRangeDepth = { VK_IMAGE_ASPECT_DEPTH_BIT | VK_IMAGE_ASPECT_STENCIL_BIT, 0, 1, 0, 1 };
 			VkImageLayout& CurrentDepthLayout = ImmediateContext.GetTransitionAndLayoutManager().FindOrAddLayoutRW(DepthTexture2D->Surface.Image, VK_IMAGE_LAYOUT_UNDEFINED);
-			bool bDepthHadLayout = (CurrentLayout != VK_IMAGE_LAYOUT_UNDEFINED);
+			bool bDepthHadLayout = (CurrentDepthLayout != VK_IMAGE_LAYOUT_UNDEFINED);
 
 			if (CurrentDepthLayout != VK_IMAGE_LAYOUT_TRANSFER_SRC_OPTIMAL)
 			{
-				vlkRHI->VulkanSetImageLayout(CmdBuffer->GetHandle(), DepthTexture2D->Surface.Image, CurrentLayout, VK_IMAGE_LAYOUT_TRANSFER_SRC_OPTIMAL, SubresourceRangeDepth);
+				vlkRHI->VulkanSetImageLayout(CmdBuffer->GetHandle(), DepthTexture2D->Surface.Image, CurrentDepthLayout, VK_IMAGE_LAYOUT_TRANSFER_SRC_OPTIMAL, SubresourceRangeDepth);
 			}
 
 			vr::VRVulkanTextureData_t VulkanTextureDataDepth{};
