Add D3D12 renderbridge implementation, including magical SteamVR bUseExplicitTimingMode flag to ensure thread-safety in DX12.

diff --git a/Engine/Plugins/Runtime/Steam/SteamVR/Source/SteamVR/Private/SteamVRHMD.cpp b/Engine/Plugins/Runtime/Steam/SteamVR/Source/SteamVR/Private/SteamVRHMD.cpp
--- a/Engine/Plugins/Runtime/Steam/SteamVR/Source/SteamVR/Private/SteamVRHMD.cpp
+++ b/Engine/Plugins/Runtime/Steam/SteamVR/Source/SteamVR/Private/SteamVRHMD.cpp
@@ -6,6 +6,7 @@
 #include "Misc/App.h"
 #include "Misc/CoreDelegates.h"
 #include "Misc/EngineVersion.h"
+#include "HardwareInfo.h"
 #include "RendererPrivate.h"
 #include "ScenePrivate.h"
 #include "Slate/SceneViewport.h"
@@ -290,6 +291,45 @@ public:
 		VRCompositor = nullptr;
 		vr::VR_Shutdown();
 	}
+
+#if PLATFORM_WINDOWS
+	enum class D3DApiLevel
+	{
+		Undefined,
+		Direct3D11,
+		Direct3D12
+	};
+
+	static inline D3DApiLevel GetD3DApiLevel()
+	{
+		FString RHIString;
+		{
+			FString HardwareDetails = FHardwareInfo::GetHardwareDetailsString();
+			FString RHILookup = NAME_RHI.ToString() + TEXT("=");
+
+			if (!FParse::Value(*HardwareDetails, *RHILookup, RHIString))
+			{
+				// RHI might not be up yet. Let's check the command-line and see if DX12 was specified.
+				// This will get hit on startup since we don't have RHI details during stereo device bringup. 
+				// This is not a final fix; we should probably move the stereo device init to later on in startup.
+				bool bForceD3D12 = FParse::Param(FCommandLine::Get(), TEXT("d3d12")) || FParse::Param(FCommandLine::Get(), TEXT("dx12"));
+				return bForceD3D12 ? D3DApiLevel::Direct3D12 : D3DApiLevel::Direct3D11;
+			}
+		}
+
+		if (RHIString == TEXT("D3D11"))
+		{
+			return D3DApiLevel::Direct3D11;
+		}
+		if (RHIString == TEXT("D3D12"))
+		{
+			return D3DApiLevel::Direct3D12;
+		}
+
+		return D3DApiLevel::Undefined;
+	}
+
+#endif
 	
 	uint64 GetGraphicsAdapterLuid() override
 	{
@@ -329,7 +369,20 @@ public:
 #if PLATFORM_WINDOWS
 				else
 				{
-					TextureType = vr::ETextureType::TextureType_DirectX;
+					D3DApiLevel level = GetD3DApiLevel();
+
+					if (level == D3DApiLevel::Direct3D11)
+					{
+						TextureType = vr::ETextureType::TextureType_DirectX;
+					}
+					else if (level == D3DApiLevel::Direct3D12)
+					{
+						TextureType = vr::ETextureType::TextureType_DirectX12;
+					}
+					else
+					{
+						return NoDevice;
+					}
 				}
 #endif
 			}
@@ -1466,7 +1519,7 @@ void FSteamVRHMD::OnBeginRendering_RenderThread(FRHICommandListImmediate& RHICmd
 
 	check(pBridge);
 	pBridge->BeginRendering_RenderThread(RHICmdList);
-	
+
 	check(SpectatorScreenController);
 	SpectatorScreenController->UpdateSpectatorScreenMode_RenderThread();
 
@@ -1779,7 +1832,16 @@ bool FSteamVRHMD::Startup()
 #if PLATFORM_WINDOWS
 			else
 			{
-				pBridge = new D3D11Bridge( this );
+				auto level = FSteamVRPlugin::GetD3DApiLevel();
+
+				if (level == FSteamVRPlugin::D3DApiLevel::Direct3D11)
+				{
+					pBridge = new D3D11Bridge(this);
+				}
+				else if (level == FSteamVRPlugin::D3DApiLevel::Direct3D12)
+				{
+					pBridge = new D3D12Bridge(this);
+				}
 			}
 #endif
 			ensure( pBridge != nullptr );
diff --git a/Engine/Plugins/Runtime/Steam/SteamVR/Source/SteamVR/Private/SteamVRHMD.h b/Engine/Plugins/Runtime/Steam/SteamVR/Source/SteamVR/Private/SteamVRHMD.h
--- a/Engine/Plugins/Runtime/Steam/SteamVR/Source/SteamVR/Private/SteamVRHMD.h
+++ b/Engine/Plugins/Runtime/Steam/SteamVR/Source/SteamVR/Private/SteamVRHMD.h
@@ -18,7 +18,6 @@
 
 #if PLATFORM_WINDOWS
 #include "Windows/AllowWindowsPlatformTypes.h"
-#include <d3d11.h>
 #include "Windows/HideWindowsPlatformTypes.h"
 #elif PLATFORM_MAC
 #include <IOSurface/IOSurface.h>
@@ -262,6 +261,16 @@ public:
 		virtual void UpdateViewport(const FViewport& Viewport, FRHIViewport* InViewportRHI) override;
 		virtual void Reset() override;
 	};
+
+	class D3D12Bridge : public BridgeBaseImpl
+	{
+	public:
+		D3D12Bridge(FSteamVRHMD* plugin);
+
+		virtual void FinishRendering() override;
+		virtual void UpdateViewport(const FViewport& Viewport, FRHIViewport* InViewportRHI) override;
+		virtual void Reset() override;
+	};
 #endif // PLATFORM_WINDOWS
 
 #if !PLATFORM_MAC
diff --git a/Engine/Plugins/Runtime/Steam/SteamVR/Source/SteamVR/Private/SteamVRRender.cpp b/Engine/Plugins/Runtime/Steam/SteamVR/Source/SteamVR/Private/SteamVRRender.cpp
--- a/Engine/Plugins/Runtime/Steam/SteamVR/Source/SteamVR/Private/SteamVRRender.cpp
+++ b/Engine/Plugins/Runtime/Steam/SteamVR/Source/SteamVR/Private/SteamVRRender.cpp
@@ -23,6 +23,10 @@
 #include "VulkanContext.h"
 #endif
 
+#if PLATFORM_WINDOWS
+#include "D3D12RHIPrivate.h"
+#endif
+
 static TAutoConsoleVariable<int32> CUsePostPresentHandoff(TEXT("vr.SteamVR.UsePostPresentHandoff"), 0, TEXT("Whether or not to use PostPresentHandoff.  If true, more GPU time will be available, but this relies on no SceneCaptureComponent2D or WidgetComponents being active in the scene.  Otherwise, it will break async reprojection."));
 
 static TAutoConsoleVariable<int> CVarEnableDepthSubmission(
@@ -178,10 +182,9 @@ void FSteamVRHMD::BridgeBaseImpl::PostPresent()
 
 #if PLATFORM_WINDOWS
 
-FSteamVRHMD::D3D11Bridge::D3D11Bridge(FSteamVRHMD* plugin):
-	BridgeBaseImpl(plugin)
+FSteamVRHMD::D3D11Bridge::D3D11Bridge(FSteamVRHMD* plugin)
+	: BridgeBaseImpl(plugin)
 {
-
 }
 
 void FSteamVRHMD::D3D11Bridge::FinishRendering()
@@ -253,6 +256,96 @@ void FSteamVRHMD::D3D11Bridge::UpdateViewport(const FViewport& Viewport, FRHIVie
 	check(RT->GetTexture2D() == SwapChain->GetTexture2D());
 }
 
+FSteamVRHMD::D3D12Bridge::D3D12Bridge(FSteamVRHMD* plugin)
+	: BridgeBaseImpl(plugin)
+{
+	bUseExplicitTimingMode = true;
+}
+
+void FSteamVRHMD::D3D12Bridge::FinishRendering()
+{
+	bool bSubmitDepth = CVarEnableDepthSubmission->GetInt() > 0;
+	vr::EVRSubmitFlags Flags = bSubmitDepth ? vr::EVRSubmitFlags::Submit_TextureWithDepth : vr::EVRSubmitFlags::Submit_Default;
+
+	auto D3D12RHI = static_cast<FD3D12DynamicRHI*>(GDynamicRHI);
+	auto Device = D3D12RHI->GetAdapter().GetDevice(0);
+
+	vr::D3D12TextureData_t TextureData;
+	TextureData.m_pResource = (ID3D12Resource*)SwapChain->GetTexture2D()->GetNativeResource();
+	TextureData.m_pCommandQueue = D3D12RHI->RHIGetD3DCommandQueue();
+	TextureData.m_nNodeMask = Device->GetGPUMask().GetNative();
+
+	vr::VRTextureWithDepth_t Texture;
+	Texture.handle = &TextureData;
+	Texture.eType = vr::TextureType_DirectX12;
+	Texture.eColorSpace = vr::ColorSpace_Auto;
+
+	// Need function scope here since we use the pointer to this is the handle for the depth texture in the struct below.
+	vr::D3D12TextureData_t DepthTextureData;
+
+	if (bSubmitDepth)
+	{
+		// If this flag is false, the struct will be treated as a vr::Texture_t and these entries will be ignored - so we can skip this if not submitting depth
+		DepthTextureData.m_pResource = (ID3D12Resource*)DepthSwapChain->GetTexture2D()->GetNativeResource();
+		DepthTextureData.m_pCommandQueue = D3D12RHI->RHIGetD3DCommandQueue();
+		DepthTextureData.m_nNodeMask = Device->GetGPUMask().GetNative();
+
+		Texture.depth.handle = &DepthTextureData;
+
+		// Set the texture depth range to follow Unreal's inverted-Z depth settings (near is 1.0f, far is 0.0f).
+		Texture.depth.vRange.v[0] = 1.0f;
+		Texture.depth.vRange.v[1] = 0.0f;
+
+		Texture.depth.mProjection = ToHmdMatrix44(Plugin->GetStereoProjectionMatrix(eSSP_LEFT_EYE));
+		// Rescale the projection (our projection value is 10.0f here, since our units are cm, and SteamVR works in meters).
+		Texture.depth.mProjection.m[2][3] *= 0.01f;
+	}
+
+	vr::VRTextureBounds_t LeftBounds;
+	LeftBounds.uMin = 0.0f;
+	LeftBounds.uMax = 0.5f;
+	LeftBounds.vMin = 0.0f;
+	LeftBounds.vMax = 1.0f;
+
+	vr::EVRCompositorError Error = Plugin->VRCompositor->Submit(vr::Eye_Left, &Texture, &LeftBounds, Flags);
+
+	vr::VRTextureBounds_t RightBounds;
+	RightBounds.uMin = 0.5f;
+	RightBounds.uMax = 1.0f;
+	RightBounds.vMin = 0.0f;
+	RightBounds.vMax = 1.0f;
+
+	if (bSubmitDepth)
+	{
+		Texture.depth.mProjection = ToHmdMatrix44(Plugin->GetStereoProjectionMatrix(eSSP_RIGHT_EYE));
+		// Rescale the projection (our projection value is 10.0f here, since our units are cm, and SteamVR works in meters).
+		Texture.depth.mProjection.m[2][3] *= 0.01f;
+	}
+
+	Error = Plugin->VRCompositor->Submit(vr::Eye_Right, &Texture, &RightBounds, Flags);
+
+	static bool FirstError = true;
+	if (FirstError && Error != vr::VRCompositorError_None)
+	{
+		UE_LOG(LogHMD, Log, TEXT("Warning: SteamVR Compositor had an error on present (%d)"), (int32)Error);
+		FirstError = false;
+	}
+}
+
+void FSteamVRHMD::D3D12Bridge::UpdateViewport(const FViewport& Viewport, FRHIViewport* InViewportRHI)
+{
+	check(IsInGameThread());
+	check(InViewportRHI);
+
+	const FTexture2DRHIRef& RT = Viewport.GetRenderTargetTexture();
+	check(IsValidRef(RT));
+	check(RT->GetTexture2D() == SwapChain->GetTexture2D());
+}
+
+void FSteamVRHMD::D3D12Bridge::Reset()
+{
+}
+
 #endif // PLATFORM_WINDOWS
 
 #if !PLATFORM_MAC
diff --git a/Engine/Plugins/Runtime/Steam/SteamVR/Source/SteamVR/SteamVR.Build.cs b/Engine/Plugins/Runtime/Steam/SteamVR/Source/SteamVR/SteamVR.Build.cs
--- a/Engine/Plugins/Runtime/Steam/SteamVR/Source/SteamVR/SteamVR.Build.cs
+++ b/Engine/Plugins/Runtime/Steam/SteamVR/Source/SteamVR/SteamVR.Build.cs
@@ -48,16 +48,34 @@ namespace UnrealBuildTool.Rules
 
             if (Target.Platform == UnrealTargetPlatform.Win32 || Target.Platform == UnrealTargetPlatform.Win64)
             {
+				PrivateDependencyModuleNames.AddRange(
+					new string[]
+					{
+							"D3D11RHI",
+							"D3D12RHI",
+					});
+
+				PrivateIncludePaths.AddRange(
+					new string[]
+					{
+							"../../../../../Source/Runtime/Windows/D3D11RHI/Private",
+							"../../../../../Source/Runtime/Windows/D3D11RHI/Private/Windows",
+							"../../../../../Source/Runtime/D3D12RHI/Private",
+							"../../../../../Source/Runtime/D3D12RHI/Private/Windows",
+					});
+
+				AddEngineThirdPartyPrivateStaticDependencies(Target, "DX12");
 				AddEngineThirdPartyPrivateStaticDependencies(Target, "OpenVR");
-                PrivateDependencyModuleNames.Add("D3D11RHI");     //@todo steamvr: multiplatform
 
-                AddEngineThirdPartyPrivateStaticDependencies(Target, "OpenGL");
+				AddEngineThirdPartyPrivateStaticDependencies(Target, "OpenGL");
                 PrivateDependencyModuleNames.Add("OpenGLDrv");
 
                 AddEngineThirdPartyPrivateStaticDependencies(Target, "Vulkan");
                 PrivateDependencyModuleNames.Add("VulkanRHI");
-            }
-            else if (Target.Platform == UnrealTargetPlatform.Mac)
+
+				AddEngineThirdPartyPrivateStaticDependencies(Target, "NVAftermath");
+			}
+			else if (Target.Platform == UnrealTargetPlatform.Mac)
             {
 				PublicFrameworks.Add("IOSurface");
                 AddEngineThirdPartyPrivateStaticDependencies(Target, "OpenVR");
