Deprecate MotionController keys

diff --git a/Engine/Plugins/Runtime/Steam/SteamVR/Source/SteamVREditor/Private/AnimGraphNode_SteamVRSetWristTransform.cpp b/Engine/Plugins/Runtime/Steam/SteamVR/Source/SteamVREditor/Private/AnimGraphNode_SteamVRSetWristTransform.cpp
new file mode 100644
--- /dev/null
+++ b/Engine/Plugins/Runtime/Steam/SteamVR/Source/SteamVREditor/Private/AnimGraphNode_SteamVRSetWristTransform.cpp
@@ -0,0 +1,64 @@
+/*
+Copyright 2019 Valve Corporation under https://opensource.org/licenses/BSD-3-Clause
+
+Redistribution and use in source and binary forms, with or without modification,
+are permitted provided that the following conditions are met:
+
+1. Redistributions of source code must retain the above copyright notice, this
+   list of conditions and the following disclaimer.
+
+2. Redistributions in binary form must reproduce the above copyright notice,
+   this list of conditions and the following disclaimer in the documentation
+   and/or other materials provided with the distribution.
+
+3. Neither the name of the copyright holder nor the names of its contributors
+   may be used to endorse or promote products derived from this software
+   without specific prior written permission.
+
+THIS SOFTWARE IS PROVIDED BY THE COPYRIGHT HOLDERS AND CONTRIBUTORS "AS IS" AND
+ANY EXPRESS OR IMPLIED WARRANTIES, INCLUDING, BUT NOT LIMITED TO, THE IMPLIED
+WARRANTIES OF MERCHANTABILITY AND FITNESS FOR A PARTICULAR PURPOSE ARE DISCLAIMED.
+IN NO EVENT SHALL THE COPYRIGHT HOLDER OR CONTRIBUTORS BE LIABLE FOR ANY DIRECT,
+INDIRECT, INCIDENTAL, SPECIAL, EXEMPLARY, OR CONSEQUENTIAL DAMAGES (INCLUDING,
+BUT NOT LIMITED TO, PROCUREMENT OF SUBSTITUTE GOODS OR SERVICES; LOSS OF USE, DATA,
+OR PROFITS; OR BUSINESS INTERRUPTION) HOWEVER CAUSED AND ON ANY THEORY OF LIABILITY,
+WHETHER IN CONTRACT, STRICT LIABILITY, OR TORT (INCLUDING NEGLIGENCE OR OTHERWISE)
+ARISING IN ANY WAY OUT OF THE USE OF THIS SOFTWARE, EVEN IF ADVISED OF THE
+POSSIBILITY OF SUCH DAMAGE.
+*/
+
+#include "AnimGraphNode_SteamVRSetWristTransform.h"
+#include "SteamVREditor.h"
+
+#define LOCTEXT_NAMESPACE "SteamVRSetWristTransformAnimNode"
+
+UAnimGraphNode_SteamVRSetWristTransform::UAnimGraphNode_SteamVRSetWristTransform(const FObjectInitializer& ObjectInitializer)
+	: Super(ObjectInitializer)
+{
+}
+
+// Node Color
+FLinearColor UAnimGraphNode_SteamVRSetWristTransform::GetNodeTitleColor() const 
+{ 
+	return FLinearColor(0.f, 0.f, 255.f, 1.f);
+}
+
+// Node Category
+FText UAnimGraphNode_SteamVRSetWristTransform::GetMenuCategory() const
+{
+	return LOCTEXT("NodeCategory", "SteamVR Input");
+}
+
+// Node Title
+FText UAnimGraphNode_SteamVRSetWristTransform::GetNodeTitle(ENodeTitleType::Type TitleType) const
+{
+	return LOCTEXT("NodeTitle", "SteamVR Set Wrist Transform");
+}
+
+// Node Tooltip
+FText UAnimGraphNode_SteamVRSetWristTransform::GetTooltipText() const
+{
+	return LOCTEXT("NodeTooltip", "Sets the wrist transform of target pose from the reference pose");
+}
+
+#undef LOCTEXT_NAMESPACE
\ No newline at end of file
diff --git a/Engine/Plugins/Runtime/Steam/SteamVR/Source/SteamVREditor/Private/SteamVREditor.cpp b/Engine/Plugins/Runtime/Steam/SteamVR/Source/SteamVREditor/Private/SteamVREditor.cpp
new file mode 100644
--- /dev/null
+++ b/Engine/Plugins/Runtime/Steam/SteamVR/Source/SteamVREditor/Private/SteamVREditor.cpp
@@ -0,0 +1,344 @@
+/*
+Copyright 2019 Valve Corporation under https://opensource.org/licenses/BSD-3-Clause
+
+Redistribution and use in source and binary forms, with or without modification,
+are permitted provided that the following conditions are met:
+
+1. Redistributions of source code must retain the above copyright notice, this 
+   list of conditions and the following disclaimer.
+
+2. Redistributions in binary form must reproduce the above copyright notice, 
+   this list of conditions and the following disclaimer in the documentation 
+   and/or other materials provided with the distribution.
+
+3. Neither the name of the copyright holder nor the names of its contributors 
+   may be used to endorse or promote products derived from this software 
+   without specific prior written permission.
+
+THIS SOFTWARE IS PROVIDED BY THE COPYRIGHT HOLDERS AND CONTRIBUTORS "AS IS" AND
+ANY EXPRESS OR IMPLIED WARRANTIES, INCLUDING, BUT NOT LIMITED TO, THE IMPLIED 
+WARRANTIES OF MERCHANTABILITY AND FITNESS FOR A PARTICULAR PURPOSE ARE DISCLAIMED. 
+IN NO EVENT SHALL THE COPYRIGHT HOLDER OR CONTRIBUTORS BE LIABLE FOR ANY DIRECT, 
+INDIRECT, INCIDENTAL, SPECIAL, EXEMPLARY, OR CONSEQUENTIAL DAMAGES (INCLUDING, 
+BUT NOT LIMITED TO, PROCUREMENT OF SUBSTITUTE GOODS OR SERVICES; LOSS OF USE, DATA, 
+OR PROFITS; OR BUSINESS INTERRUPTION) HOWEVER CAUSED AND ON ANY THEORY OF LIABILITY, 
+WHETHER IN CONTRACT, STRICT LIABILITY, OR TORT (INCLUDING NEGLIGENCE OR OTHERWISE) 
+ARISING IN ANY WAY OUT OF THE USE OF THIS SOFTWARE, EVEN IF ADVISED OF THE 
+POSSIBILITY OF SUCH DAMAGE.
+*/
+
+#pragma once
+
+#include "SteamVREditor.h"
+#include "IMotionController.h"
+#include "Features/IModularFeatures.h"
+#include "SteamVREditorStyle.h"
+#include "SteamVREditorCommands.h"
+#include "SteamVRControllerKeys.h"
+#include "LevelEditor.h"
+
+static const FName SteamVREditorTabName("SteamVREditor");
+
+#define LOCTEXT_NAMESPACE "FSteamVREditorModule"
+
+void FSteamVREditorModule::StartupModule()
+{
+	FSteamVREditorStyle::Initialize();
+	FSteamVREditorStyle::ReloadTextures();
+
+	FSteamVREditorCommands::Register();
+	
+	PluginCommands = MakeShareable(new FUICommandList);
+
+	// Dummy action for main toolbar button
+	PluginCommands->MapAction(
+		FSteamVREditorCommands::Get().PluginAction,
+		FExecuteAction::CreateRaw(this, &FSteamVREditorModule::PluginButtonClicked),
+		FCanExecuteAction());
+	
+	// Regenerate Action Manifest
+	PluginCommands->MapAction(
+		FSteamVREditorCommands::Get().JsonActionManifest,
+		FExecuteAction::CreateRaw(this, &FSteamVREditorModule::JsonRegenerateActionManifest),
+		FCanExecuteAction());
+
+	// Regenerate Controller Bindings
+	PluginCommands->MapAction(
+		FSteamVREditorCommands::Get().JsonControllerBindings,
+		FExecuteAction::CreateRaw(this, &FSteamVREditorModule::JsonRegenerateControllerBindings),
+		FCanExecuteAction());
+
+	// Reload Action Manifest
+	PluginCommands->MapAction(
+		FSteamVREditorCommands::Get().ReloadActionManifest,
+		FExecuteAction::CreateRaw(this, &FSteamVREditorModule::ReloadActionManifest),
+		FCanExecuteAction());
+
+	// Launch Bindings URL
+	PluginCommands->MapAction(
+		FSteamVREditorCommands::Get().LaunchBindingsURL,
+		FExecuteAction::CreateRaw(this, &FSteamVREditorModule::LaunchBindingsURL),
+		FCanExecuteAction());
+	
+	// Add Sample Inputs
+	PluginCommands->MapAction(
+		FSteamVREditorCommands::Get().AddSampleInputs,
+		FExecuteAction::CreateRaw(this, &FSteamVREditorModule::AddSampleInputs),
+		FCanExecuteAction());
+
+	FLevelEditorModule& LevelEditorModule = FModuleManager::LoadModuleChecked<FLevelEditorModule>("LevelEditor");
+	
+	{
+		TSharedPtr<FExtender> MenuExtender = MakeShareable(new FExtender());
+		MenuExtender->AddMenuExtension("WindowLayout", EExtensionHook::After, PluginCommands, FMenuExtensionDelegate::CreateRaw(this, &FSteamVREditorModule::AddMenuExtension));
+
+		LevelEditorModule.GetMenuExtensibilityManager()->AddExtender(MenuExtender);
+	}
+	
+	{
+		TSharedPtr<FExtender> ToolbarExtender = MakeShareable(new FExtender);
+		ToolbarExtender->AddToolBarExtension("Settings", EExtensionHook::After, PluginCommands, FToolBarExtensionDelegate::CreateRaw(this, &FSteamVREditorModule::AddToolbarExtension));
+		
+		LevelEditorModule.GetToolBarExtensibilityManager()->AddExtender(ToolbarExtender);
+	}
+}
+
+void FSteamVREditorModule::ShutdownModule()
+{
+	FSteamVREditorStyle::Shutdown();
+
+	FSteamVREditorCommands::Unregister();
+}
+
+void FSteamVREditorModule::PluginButtonClicked()
+{
+	// Empty on purpose
+}
+
+void FSteamVREditorModule::JsonRegenerateActionManifest()
+{
+	USteamVRInputDeviceFunctionLibrary::RegenActionManifest();
+}
+
+void FSteamVREditorModule::JsonRegenerateControllerBindings()
+{
+	USteamVRInputDeviceFunctionLibrary::RegenControllerBindings();
+}
+
+void FSteamVREditorModule::ReloadActionManifest()
+{
+	USteamVRInputDeviceFunctionLibrary::ReloadActionManifest();
+}
+
+void FSteamVREditorModule::LaunchBindingsURL()
+{
+	USteamVRInputDeviceFunctionLibrary::LaunchBindingsURL();
+}
+
+void FSteamVREditorModule::AddSampleInputs()
+{
+	//// Get Existing Input Settings
+	//auto DefaultInputSettings = GetDefault<UInputSettings>();
+	//TArray<FInputAxisKeyMapping> ExistingAxisKeys = DefaultInputSettings->GetAxisMappings();
+	//TArray<FInputActionKeyMapping> ExistingActionKeys = DefaultInputSettings->GetActionMappings();
+
+	//// Create new Input Settings
+	//UInputSettings* InputSettings = GetMutableDefault<UInputSettings>();
+
+	//if (InputSettings->IsValidLowLevel())
+	//{
+	//	// Teleport
+	//	AddUniqueActionMapping(ExistingActionKeys, InputSettings, FName(TEXT("TeleportLeft")), EKeys::Valve_Index_Controller_Trackpad_Press_Left);
+	//	AddUniqueActionMapping(ExistingActionKeys, InputSettings, FName(TEXT("TeleportLeft")), EKeys::HTC_Vive_Controller_Trackpad_Press_Left);
+	//	AddUniqueActionMapping(ExistingActionKeys, InputSettings, FName(TEXT("TeleportLeft")), EKeys::HTC_Cosmos_Controller_Joystick_Press_Left);
+	//	AddUniqueActionMapping(ExistingActionKeys, InputSettings, FName(TEXT("TeleportLeft")), EKeys::Oculus_Touch_Controller_Joystick_Press_Left);
+	//	AddUniqueActionMapping(ExistingActionKeys, InputSettings, FName(TEXT("TeleportLeft")), EKeys::Windows_MR_Controller_Trackpad_Press_Left);
+	//	AddUniqueActionMapping(ExistingActionKeys, InputSettings, FName(TEXT("TeleportLeft")), FGamepadKeyNames::MotionController_Left_Thumbstick);
+
+	//	AddUniqueActionMapping(ExistingActionKeys, InputSettings, FName(TEXT("TeleportRight")), EKeys::Valve_Index_Controller_Trackpad_Press_Right);
+	//	AddUniqueActionMapping(ExistingActionKeys, InputSettings, FName(TEXT("TeleportRight")), EKeys::HTC_Vive_Controller_Trackpad_Press_Right);
+	//	AddUniqueActionMapping(ExistingActionKeys, InputSettings, FName(TEXT("TeleportRight")), EKeys::HTC_Cosmos_Controller_Joystick_Press_Right);
+	//	AddUniqueActionMapping(ExistingActionKeys, InputSettings, FName(TEXT("TeleportRight")), EKeys::Oculus_Touch_Controller_Joystick_Press_Right);
+	//	AddUniqueActionMapping(ExistingActionKeys, InputSettings, FName(TEXT("TeleportRight")), EKeys::Windows_MR_Controller_Trackpad_Press_Right);
+	//	AddUniqueActionMapping(ExistingActionKeys, InputSettings, FName(TEXT("TeleportRight")), FGamepadKeyNames::MotionController_Right_Thumbstick);
+
+	//	// Grab
+	//	AddUniqueActionMapping(ExistingActionKeys, InputSettings, FName(TEXT("GrabLeft")), EKeys::Valve_Index_Controller_Grip_Grab_Left);
+	//	AddUniqueActionMapping(ExistingActionKeys, InputSettings, FName(TEXT("GrabLeft")), EKeys::HTC_Vive_Controller_Trigger_Press_Left);
+	//	AddUniqueActionMapping(ExistingActionKeys, InputSettings, FName(TEXT("GrabLeft")), EKeys::HTC_Cosmos_Controller_Grip_Click_Left);
+	//	AddUniqueActionMapping(ExistingActionKeys, InputSettings, FName(TEXT("GrabLeft")), EKeys::Oculus_Touch_Controller_Grip_Press_Left);
+	//	AddUniqueActionMapping(ExistingActionKeys, InputSettings, FName(TEXT("GrabLeft")), EKeys::Windows_MR_Controller_Grip_Press_Left);
+	//	AddUniqueActionMapping(ExistingActionKeys, InputSettings, FName(TEXT("GrabLeft")), FGamepadKeyNames::MotionController_Left_Trigger);
+
+	//	AddUniqueActionMapping(ExistingActionKeys, InputSettings, FName(TEXT("GrabRight")), EKeys::Valve_Index_Controller_Grip_Grab_Right);
+	//	AddUniqueActionMapping(ExistingActionKeys, InputSettings, FName(TEXT("GrabRight")), EKeys::HTC_Vive_Controller_Trigger_Press_Right);
+	//	AddUniqueActionMapping(ExistingActionKeys, InputSettings, FName(TEXT("GrabRight")), EKeys::HTC_Cosmos_Controller_Grip_Click_Right);
+	//	AddUniqueActionMapping(ExistingActionKeys, InputSettings, FName(TEXT("GrabRight")), EKeys::Oculus_Touch_Controller_Grip_Press_Right);
+	//	AddUniqueActionMapping(ExistingActionKeys, InputSettings, FName(TEXT("GrabRight")), EKeys::Windows_MR_Controller_Grip_Press_Right);
+	//	AddUniqueActionMapping(ExistingActionKeys, InputSettings, FName(TEXT("GrabRight")), FGamepadKeyNames::MotionController_Right_Trigger);
+
+	//	// Fire Arrow
+	//	AddUniqueActionMapping(ExistingActionKeys, InputSettings, FName(TEXT("FireArrowLeft")), EKeys::Valve_Index_Controller_Pinch_Grab_Left);
+	//	AddUniqueActionMapping(ExistingActionKeys, InputSettings, FName(TEXT("FireArrowLeft")), EKeys::HTC_Vive_Controller_Trigger_Press_Left);
+	//	AddUniqueActionMapping(ExistingActionKeys, InputSettings, FName(TEXT("FireArrowLeft")), EKeys::HTC_Cosmos_Controller_Trigger_Press_Left);
+	//	AddUniqueActionMapping(ExistingActionKeys, InputSettings, FName(TEXT("FireArrowLeft")), EKeys::Oculus_Touch_Controller_Trigger_Press_Left);
+	//	AddUniqueActionMapping(ExistingActionKeys, InputSettings, FName(TEXT("FireArrowLeft")), EKeys::Windows_MR_Controller_Trigger_Press_Left);
+	//	AddUniqueActionMapping(ExistingActionKeys, InputSettings, FName(TEXT("FireArrowLeft")), FGamepadKeyNames::MotionController_Left_Trigger);
+
+	//	AddUniqueActionMapping(ExistingActionKeys, InputSettings, FName(TEXT("FireArrowRight")), EKeys::Valve_Index_Controller_Pinch_Grab_Right);
+	//	AddUniqueActionMapping(ExistingActionKeys, InputSettings, FName(TEXT("FireArrowRight")), EKeys::HTC_Vive_Controller_Trigger_Press_Right);
+	//	AddUniqueActionMapping(ExistingActionKeys, InputSettings, FName(TEXT("FireArrowRight")), EKeys::HTC_Cosmos_Controller_Trigger_Press_Right);
+	//	AddUniqueActionMapping(ExistingActionKeys, InputSettings, FName(TEXT("FireArrowRight")), EKeys::Oculus_Touch_Controller_Trigger_Press_Right);
+	//	AddUniqueActionMapping(ExistingActionKeys, InputSettings, FName(TEXT("FireArrowRight")), EKeys::Windows_MR_Controller_Trigger_Press_Right);
+	//	AddUniqueActionMapping(ExistingActionKeys, InputSettings, FName(TEXT("FireArrowRight")), FGamepadKeyNames::MotionController_Right_Trigger);
+
+	//	// HMD Proximity
+	//	AddUniqueActionMapping(ExistingActionKeys, InputSettings, FName(TEXT("HeadsetOn")), GenericKeys::SteamVR_HMD_Proximity);
+
+	//	// Move Right
+	//	AddUniqueAxisMapping(ExistingAxisKeys, InputSettings, FName(TEXT("MoveRight_Y")), EKeys::Valve_Index_Controller_Thumbstick_Y_Right);
+	//	AddUniqueAxisMapping(ExistingAxisKeys, InputSettings, FName(TEXT("MoveRight_Y")), EKeys::HTC_Vive_Controller_Trackpad_Y_Right);
+	//	AddUniqueAxisMapping(ExistingAxisKeys, InputSettings, FName(TEXT("MoveRight_Y")), EKeys::HTC_Cosmos_Controller_Joystick_Y_Right);
+	//	AddUniqueAxisMapping(ExistingAxisKeys, InputSettings, FName(TEXT("MoveRight_Y")), EKeys::Oculus_Touch_Controller_Joystick_Y_Right);
+	//	AddUniqueAxisMapping(ExistingAxisKeys, InputSettings, FName(TEXT("MoveRight_Y")), EKeys::Windows_MR_Controller_Joystick_Y_Right);
+	//	AddUniqueAxisMapping(ExistingAxisKeys, InputSettings, FName(TEXT("MoveRight_Y")), FGamepadKeyNames::MotionController_Right_Thumbstick_Y);
+
+	//	AddUniqueAxisMapping(ExistingAxisKeys, InputSettings, FName(TEXT("MoveRight_X")), EKeys::Valve_Index_Controller_Thumbstick_X_Right);
+	//	AddUniqueAxisMapping(ExistingAxisKeys, InputSettings, FName(TEXT("MoveRight_X")), EKeys::HTC_Vive_Controller_Trackpad_X_Right);
+	//	AddUniqueAxisMapping(ExistingAxisKeys, InputSettings, FName(TEXT("MoveRight_X")), EKeys::HTC_Cosmos_Controller_Joystick_X_Right);
+	//	AddUniqueAxisMapping(ExistingAxisKeys, InputSettings, FName(TEXT("MoveRight_X")), EKeys::Oculus_Touch_Controller_Joystick_X_Right);
+	//	AddUniqueAxisMapping(ExistingAxisKeys, InputSettings, FName(TEXT("MoveRight_X")), EKeys::Windows_MR_Controller_Joystick_X_Right);
+	//	AddUniqueAxisMapping(ExistingAxisKeys, InputSettings, FName(TEXT("MoveRight_X")), FGamepadKeyNames::MotionController_Right_Thumbstick_X);
+
+	//	AddUniqueAxisMapping(ExistingAxisKeys, InputSettings, FName(TEXT("MoveLeft_Y")), EKeys::Valve_Index_Controller_Thumbstick_Y_Left);
+	//	AddUniqueAxisMapping(ExistingAxisKeys, InputSettings, FName(TEXT("MoveLeft_Y")), EKeys::HTC_Vive_Controller_Trackpad_Y_Left);
+	//	AddUniqueAxisMapping(ExistingAxisKeys, InputSettings, FName(TEXT("MoveLeft_Y")), EKeys::HTC_Cosmos_Controller_Joystick_Y_Left);
+	//	AddUniqueAxisMapping(ExistingAxisKeys, InputSettings, FName(TEXT("MoveLeft_Y")), EKeys::Oculus_Touch_Controller_Joystick_Y_Left);
+	//	AddUniqueAxisMapping(ExistingAxisKeys, InputSettings, FName(TEXT("MoveLeft_Y")), EKeys::Windows_MR_Controller_Joystick_Y_Left);
+	//	AddUniqueAxisMapping(ExistingAxisKeys, InputSettings, FName(TEXT("MoveLeft_Y")), FGamepadKeyNames::MotionController_Left_Thumbstick_Y);
+
+	//	AddUniqueAxisMapping(ExistingAxisKeys, InputSettings, FName(TEXT("MoveLeft_X")), EKeys::Valve_Index_Controller_Thumbstick_X_Left);
+	//	AddUniqueAxisMapping(ExistingAxisKeys, InputSettings, FName(TEXT("MoveLeft_X")), EKeys::HTC_Vive_Controller_Trackpad_X_Left);
+	//	AddUniqueAxisMapping(ExistingAxisKeys, InputSettings, FName(TEXT("MoveLeft_X")), EKeys::HTC_Cosmos_Controller_Joystick_X_Left);
+	//	AddUniqueAxisMapping(ExistingAxisKeys, InputSettings, FName(TEXT("MoveLeft_X")), EKeys::Oculus_Touch_Controller_Joystick_X_Left);
+	//	AddUniqueAxisMapping(ExistingAxisKeys, InputSettings, FName(TEXT("MoveLeft_X")), EKeys::Windows_MR_Controller_Joystick_X_Left);
+	//	AddUniqueAxisMapping(ExistingAxisKeys, InputSettings, FName(TEXT("MoveLeft_X")), FGamepadKeyNames::MotionController_Left_Thumbstick_X);
+
+	//	// Teleport Direction
+	//	AddUniqueAxisMapping(ExistingAxisKeys, InputSettings, FName(TEXT("TeleportDirectionRight_Y")), EKeys::Valve_Index_Controller_Trackpad_Y_Right);
+	//	AddUniqueAxisMapping(ExistingAxisKeys, InputSettings, FName(TEXT("TeleportDirectionRight_Y")), EKeys::HTC_Vive_Controller_Trackpad_Y_Right);
+	//	AddUniqueAxisMapping(ExistingAxisKeys, InputSettings, FName(TEXT("TeleportDirectionRight_Y")), EKeys::HTC_Cosmos_Controller_Joystick_Y_Right);
+	//	AddUniqueAxisMapping(ExistingAxisKeys, InputSettings, FName(TEXT("TeleportDirectionRight_Y")), EKeys::Oculus_Touch_Controller_Joystick_Y_Right);
+	//	AddUniqueAxisMapping(ExistingAxisKeys, InputSettings, FName(TEXT("TeleportDirectionRight_Y")), EKeys::Windows_MR_Controller_Trackpad_Y_Right);
+	//	AddUniqueAxisMapping(ExistingAxisKeys, InputSettings, FName(TEXT("TeleportDirectionRight_Y")), FGamepadKeyNames::MotionController_Right_Thumbstick_Y);
+
+	//	AddUniqueAxisMapping(ExistingAxisKeys, InputSettings, FName(TEXT("TeleportDirectionRight_X")), EKeys::Valve_Index_Controller_Trackpad_X_Right);
+	//	AddUniqueAxisMapping(ExistingAxisKeys, InputSettings, FName(TEXT("TeleportDirectionRight_X")), EKeys::HTC_Vive_Controller_Trackpad_X_Right);
+	//	AddUniqueAxisMapping(ExistingAxisKeys, InputSettings, FName(TEXT("TeleportDirectionRight_X")), EKeys::HTC_Cosmos_Controller_Joystick_X_Right);
+	//	AddUniqueAxisMapping(ExistingAxisKeys, InputSettings, FName(TEXT("TeleportDirectionRight_X")), EKeys::Oculus_Touch_Controller_Joystick_X_Right);
+	//	AddUniqueAxisMapping(ExistingAxisKeys, InputSettings, FName(TEXT("TeleportDirectionRight_X")), EKeys::Windows_MR_Controller_Trackpad_X_Right);
+	//	AddUniqueAxisMapping(ExistingAxisKeys, InputSettings, FName(TEXT("TeleportDirectionRight_X")), FGamepadKeyNames::MotionController_Right_Thumbstick_X);
+
+	//	AddUniqueAxisMapping(ExistingAxisKeys, InputSettings, FName(TEXT("TeleportDirectionLeft_Y")), EKeys::Valve_Index_Controller_Trackpad_Y_Left);
+	//	AddUniqueAxisMapping(ExistingAxisKeys, InputSettings, FName(TEXT("TeleportDirectionLeft_Y")), EKeys::HTC_Vive_Controller_Trackpad_Y_Left);
+	//	AddUniqueAxisMapping(ExistingAxisKeys, InputSettings, FName(TEXT("TeleportDirectionLeft_Y")), EKeys::HTC_Cosmos_Controller_Joystick_Y_Left);
+	//	AddUniqueAxisMapping(ExistingAxisKeys, InputSettings, FName(TEXT("TeleportDirectionLeft_Y")), EKeys::Oculus_Touch_Controller_Joystick_Y_Left);
+	//	AddUniqueAxisMapping(ExistingAxisKeys, InputSettings, FName(TEXT("TeleportDirectionLeft_Y")), EKeys::Windows_MR_Controller_Trackpad_Y_Left);
+	//	AddUniqueAxisMapping(ExistingAxisKeys, InputSettings, FName(TEXT("TeleportDirectionLeft_Y")), FGamepadKeyNames::MotionController_Left_Thumbstick_Y);
+
+	//	AddUniqueAxisMapping(ExistingAxisKeys, InputSettings, FName(TEXT("TeleportDirectionLeft_X")), EKeys::Valve_Index_Controller_Trackpad_X_Left);
+	//	AddUniqueAxisMapping(ExistingAxisKeys, InputSettings, FName(TEXT("TeleportDirectionLeft_X")), EKeys::HTC_Vive_Controller_Trackpad_X_Left);
+	//	AddUniqueAxisMapping(ExistingAxisKeys, InputSettings, FName(TEXT("TeleportDirectionLeft_X")), EKeys::HTC_Cosmos_Controller_Joystick_X_Left);
+	//	AddUniqueAxisMapping(ExistingAxisKeys, InputSettings, FName(TEXT("TeleportDirectionLeft_X")), EKeys::Oculus_Touch_Controller_Joystick_X_Left);
+	//	AddUniqueAxisMapping(ExistingAxisKeys, InputSettings, FName(TEXT("TeleportDirectionLeft_X")), EKeys::Windows_MR_Controller_Trackpad_X_Left);
+	//	AddUniqueAxisMapping(ExistingAxisKeys, InputSettings, FName(TEXT("TeleportDirectionLeft_X")), FGamepadKeyNames::MotionController_Left_Thumbstick_X);
+
+	//	// Squeeze
+	//	AddUniqueAxisMapping(ExistingAxisKeys, InputSettings, FName(TEXT("SqueezeLeft")), EKeys::Valve_Index_Controller_GripForce_Axis_Left);
+	//	AddUniqueAxisMapping(ExistingAxisKeys, InputSettings, FName(TEXT("SqueezeLeft")), EKeys::HTC_Vive_Controller_Trigger_Pull_Left);
+	//	AddUniqueAxisMapping(ExistingAxisKeys, InputSettings, FName(TEXT("SqueezeLeft")), EKeys::HTC_Cosmos_Controller_Trigger_Pull_Left);
+	//	AddUniqueAxisMapping(ExistingAxisKeys, InputSettings, FName(TEXT("SqueezeLeft")), EKeys::Oculus_Touch_Controller_Trigger_Pull_Left);
+	//	AddUniqueAxisMapping(ExistingAxisKeys, InputSettings, FName(TEXT("SqueezeLeft")), EKeys::Windows_MR_Controller_Trigger_Pull_Left);
+	//	AddUniqueAxisMapping(ExistingAxisKeys, InputSettings, FName(TEXT("SqueezeLeft")), FGamepadKeyNames::MotionController_Left_TriggerAxis);
+
+	//	AddUniqueAxisMapping(ExistingAxisKeys, InputSettings, FName(TEXT("SqueezeRight")), EKeys::Valve_Index_Controller_GripForce_Axis_Right);
+	//	AddUniqueAxisMapping(ExistingAxisKeys, InputSettings, FName(TEXT("SqueezeRight")), EKeys::HTC_Vive_Controller_Trigger_Pull_Right);
+	//	AddUniqueAxisMapping(ExistingAxisKeys, InputSettings, FName(TEXT("SqueezeRight")), EKeys::HTC_Cosmos_Controller_Trigger_Pull_Right);
+	//	AddUniqueAxisMapping(ExistingAxisKeys, InputSettings, FName(TEXT("SqueezeRight")), EKeys::Oculus_Touch_Controller_Trigger_Pull_Right);
+	//	AddUniqueAxisMapping(ExistingAxisKeys, InputSettings, FName(TEXT("SqueezeRight")), EKeys::Windows_MR_Controller_Trigger_Pull_Right);
+	//	AddUniqueAxisMapping(ExistingAxisKeys, InputSettings, FName(TEXT("SqueezeRight")), FGamepadKeyNames::MotionController_Right_TriggerAxis);
+	//	
+
+	//	// Update the config file
+	//	InputSettings->SaveKeyMappings();
+	//	InputSettings->UpdateDefaultConfigFile();
+	//}
+}
+
+bool FSteamVREditorModule::AddUniqueAxisMapping(TArray<FInputAxisKeyMapping> ExistingAxisKeys, UInputSettings* InputSettings, FName ActionName, FKey ActionKey)
+{
+	// Create new axis mapping
+	FInputAxisKeyMapping NewAxisMapping = FInputAxisKeyMapping(ActionName, ActionKey);
+
+	// Check if this mapping already exists in the project
+	if (ExistingAxisKeys.Find(NewAxisMapping) < 1)
+	{
+		// If none, create a new one
+		InputSettings->AddAxisMapping(NewAxisMapping);
+		return true;
+	}
+	
+return false;
+}
+
+bool FSteamVREditorModule::AddUniqueActionMapping(TArray<FInputActionKeyMapping> ExistingActionKeys, UInputSettings* InputSettings, FName ActionName, FKey ActionKey)
+{
+	// Create new action mapping
+	FInputActionKeyMapping NewActionMapping = FInputActionKeyMapping(ActionName, ActionKey);
+
+	// Check if this mapping already exists in the project
+	if (ExistingActionKeys.Find(NewActionMapping) < 1)
+	{
+		// If none, create a new one
+		InputSettings->AddActionMapping(NewActionMapping);
+		return true;
+	}
+
+	return false;
+}
+
+void FSteamVREditorModule::AddMenuExtension(FMenuBuilder& Builder)
+{
+	Builder.AddMenuEntry(FSteamVREditorCommands::Get().PluginAction);
+}
+
+void FSteamVREditorModule::AddToolbarExtension(FToolBarBuilder& Builder)
+{
+	FSteamVREditorStyle MenuStyle = FSteamVREditorStyle();
+	MenuStyle.Initialize();
+
+	Builder.AddComboButton(
+		FUIAction(FExecuteAction::CreateRaw(this, &FSteamVREditorModule::PluginButtonClicked)),
+		FOnGetContent::CreateRaw(this, &FSteamVREditorModule::FillComboButton, PluginCommands),
+		LOCTEXT("SteamVRInputBtn", "SteamVR Input"),
+		LOCTEXT("SteamVRInputBtnTootlip", "SteamVR Input"),
+		FSlateIcon(FSteamVREditorStyle::GetStyleSetName(), "SteamVREditor.PluginAction")
+	);
+}
+
+TSharedRef<SWidget> FSteamVREditorModule::FillComboButton(TSharedPtr<class FUICommandList> Commands)
+{
+	FMenuBuilder MenuBuilder(true, Commands);
+
+	MenuBuilder.AddMenuEntry(FSteamVREditorCommands::Get().JsonActionManifest, NAME_None, TAttribute<FText>(), TAttribute<FText>(), FSlateIcon(FSteamVREditorStyle::GetStyleSetName(), "SteamVREditor.JsonActionManifest"));
+	MenuBuilder.AddMenuEntry(FSteamVREditorCommands::Get().JsonControllerBindings, NAME_None, TAttribute<FText>(), TAttribute<FText>(), FSlateIcon(FSteamVREditorStyle::GetStyleSetName(), "SteamVREditor.JsonControllerBindings"));
+	MenuBuilder.AddMenuEntry(FSteamVREditorCommands::Get().ReloadActionManifest, NAME_None, TAttribute<FText>(), TAttribute<FText>(), FSlateIcon(FSteamVREditorStyle::GetStyleSetName(), "SteamVREditor.ReloadActionManifest"));
+	MenuBuilder.AddMenuEntry(FSteamVREditorCommands::Get().LaunchBindingsURL, NAME_None, TAttribute<FText>(), TAttribute<FText>(), FSlateIcon(FSteamVREditorStyle::GetStyleSetName(), "SteamVREditor.LaunchBindingsURL"));
+	//MenuBuilder.AddMenuEntry(FSteamVREditorCommands::Get().AddSampleInputs, NAME_None, TAttribute<FText>(), TAttribute<FText>(), FSlateIcon(FSteamVREditorStyle::GetStyleSetName(), "SteamVREditor.AddSampleInputs"));
+
+	return MenuBuilder.MakeWidget();
+}
+
+#undef LOCTEXT_NAMESPACE
+	
+IMPLEMENT_MODULE(FSteamVREditorModule, SteamVREditor)
\ No newline at end of file
diff --git a/Engine/Plugins/Runtime/Steam/SteamVR/Source/SteamVREditor/Private/SteamVREditorStyle.cpp b/Engine/Plugins/Runtime/Steam/SteamVR/Source/SteamVREditor/Private/SteamVREditorStyle.cpp
new file mode 100644
--- /dev/null
+++ b/Engine/Plugins/Runtime/Steam/SteamVR/Source/SteamVREditor/Private/SteamVREditorStyle.cpp
@@ -0,0 +1,100 @@
+/*
+Copyright 2019 Valve Corporation under https://opensource.org/licenses/BSD-3-Clause
+
+Redistribution and use in source and binary forms, with or without modification,
+are permitted provided that the following conditions are met:
+
+1. Redistributions of source code must retain the above copyright notice, this
+   list of conditions and the following disclaimer.
+
+2. Redistributions in binary form must reproduce the above copyright notice,
+   this list of conditions and the following disclaimer in the documentation
+   and/or other materials provided with the distribution.
+
+3. Neither the name of the copyright holder nor the names of its contributors
+   may be used to endorse or promote products derived from this software
+   without specific prior written permission.
+
+THIS SOFTWARE IS PROVIDED BY THE COPYRIGHT HOLDERS AND CONTRIBUTORS "AS IS" AND
+ANY EXPRESS OR IMPLIED WARRANTIES, INCLUDING, BUT NOT LIMITED TO, THE IMPLIED
+WARRANTIES OF MERCHANTABILITY AND FITNESS FOR A PARTICULAR PURPOSE ARE DISCLAIMED.
+IN NO EVENT SHALL THE COPYRIGHT HOLDER OR CONTRIBUTORS BE LIABLE FOR ANY DIRECT,
+INDIRECT, INCIDENTAL, SPECIAL, EXEMPLARY, OR CONSEQUENTIAL DAMAGES (INCLUDING,
+BUT NOT LIMITED TO, PROCUREMENT OF SUBSTITUTE GOODS OR SERVICES; LOSS OF USE, DATA,
+OR PROFITS; OR BUSINESS INTERRUPTION) HOWEVER CAUSED AND ON ANY THEORY OF LIABILITY,
+WHETHER IN CONTRACT, STRICT LIABILITY, OR TORT (INCLUDING NEGLIGENCE OR OTHERWISE)
+ARISING IN ANY WAY OUT OF THE USE OF THIS SOFTWARE, EVEN IF ADVISED OF THE
+POSSIBILITY OF SUCH DAMAGE.
+*/
+
+#include "SteamVREditorStyle.h"
+#include "SteamVREditor.h"
+#include "Runtime/Projects/Public/Interfaces/IPluginManager.h"
+
+TSharedPtr< FSlateStyleSet > FSteamVREditorStyle::StyleInstance = NULL;
+
+void FSteamVREditorStyle::Initialize()
+{
+	if (!StyleInstance.IsValid())
+	{
+		StyleInstance = Create();
+		FSlateStyleRegistry::RegisterSlateStyle(*StyleInstance);
+	}
+}
+
+void FSteamVREditorStyle::Shutdown()
+{
+	FSlateStyleRegistry::UnRegisterSlateStyle(*StyleInstance);
+	ensure(StyleInstance.IsUnique());
+	StyleInstance.Reset();
+}
+
+FName FSteamVREditorStyle::GetStyleSetName()
+{
+	static FName StyleSetName(TEXT("SteamVREditorStyle"));
+	return StyleSetName;
+}
+
+#define IMAGE_BRUSH( RelativePath, ... ) FSlateImageBrush( Style->RootToContentDir( RelativePath, TEXT(".png") ), __VA_ARGS__ )
+#define BOX_BRUSH( RelativePath, ... ) FSlateBoxBrush( Style->RootToContentDir( RelativePath, TEXT(".png") ), __VA_ARGS__ )
+#define BORDER_BRUSH( RelativePath, ... ) FSlateBorderBrush( Style->RootToContentDir( RelativePath, TEXT(".png") ), __VA_ARGS__ )
+#define TTF_FONT( RelativePath, ... ) FSlateFontInfo( Style->RootToContentDir( RelativePath, TEXT(".ttf") ), __VA_ARGS__ )
+#define OTF_FONT( RelativePath, ... ) FSlateFontInfo( Style->RootToContentDir( RelativePath, TEXT(".otf") ), __VA_ARGS__ )
+
+const FVector2D Icon16x16(16.0f, 16.0f);
+const FVector2D Icon20x20(20.0f, 20.0f);
+const FVector2D Icon40x40(40.0f, 40.0f);
+
+TSharedRef< FSlateStyleSet > FSteamVREditorStyle::Create()
+{
+	TSharedRef< FSlateStyleSet > Style = MakeShareable(new FSlateStyleSet("SteamVREditorStyle"));
+	Style->SetContentRoot(IPluginManager::Get().FindPlugin("SteamVR")->GetBaseDir() / TEXT("Resources"));
+
+	Style->Set("SteamVREditor.PluginAction", new IMAGE_BRUSH(TEXT("Icon_PluginAction"), Icon40x40));
+	Style->Set("SteamVREditor.JsonActionManifest", new IMAGE_BRUSH(TEXT("Icon_regenerate_action_manifest"), Icon16x16));
+	Style->Set("SteamVREditor.JsonControllerBindings", new IMAGE_BRUSH(TEXT("Icon_regenerate_controller_bindings"), Icon16x16));
+	Style->Set("SteamVREditor.ReloadActionManifest", new IMAGE_BRUSH(TEXT("Icon_reload_action_manifest"), Icon16x16));
+	Style->Set("SteamVREditor.LaunchBindingsURL", new IMAGE_BRUSH(TEXT("Icon_launch_bindings_url"), Icon16x16));
+	//Style->Set("SteamVREditor.AddSampleInputs", new IMAGE_BRUSH(TEXT("Icon_add_sample_inputs"), Icon16x16));
+
+	return Style;
+}
+
+#undef IMAGE_BRUSH
+#undef BOX_BRUSH
+#undef BORDER_BRUSH
+#undef TTF_FONT
+#undef OTF_FONT
+
+void FSteamVREditorStyle::ReloadTextures()
+{
+	if (FSlateApplication::IsInitialized())
+	{
+		FSlateApplication::Get().GetRenderer()->ReloadTextureResources();
+	}
+}
+
+const ISlateStyle& FSteamVREditorStyle::Get()
+{
+	return *StyleInstance;
+}
diff --git a/Engine/Plugins/Runtime/Steam/SteamVR/Source/SteamVREditor/Public/AnimGraphNode_SteamVRInputAnimPose.h b/Engine/Plugins/Runtime/Steam/SteamVR/Source/SteamVREditor/Public/AnimGraphNode_SteamVRInputAnimPose.h
new file mode 100644
--- /dev/null
+++ b/Engine/Plugins/Runtime/Steam/SteamVR/Source/SteamVREditor/Public/AnimGraphNode_SteamVRInputAnimPose.h
@@ -0,0 +1,55 @@
+/*
+Copyright 2019 Valve Corporation under https://opensource.org/licenses/BSD-3-Clause
+
+Redistribution and use in source and binary forms, with or without modification,
+are permitted provided that the following conditions are met:
+
+1. Redistributions of source code must retain the above copyright notice, this
+   list of conditions and the following disclaimer.
+
+2. Redistributions in binary form must reproduce the above copyright notice,
+   this list of conditions and the following disclaimer in the documentation
+   and/or other materials provided with the distribution.
+
+3. Neither the name of the copyright holder nor the names of its contributors
+   may be used to endorse or promote products derived from this software
+   without specific prior written permission.
+
+THIS SOFTWARE IS PROVIDED BY THE COPYRIGHT HOLDERS AND CONTRIBUTORS "AS IS" AND
+ANY EXPRESS OR IMPLIED WARRANTIES, INCLUDING, BUT NOT LIMITED TO, THE IMPLIED
+WARRANTIES OF MERCHANTABILITY AND FITNESS FOR A PARTICULAR PURPOSE ARE DISCLAIMED.
+IN NO EVENT SHALL THE COPYRIGHT HOLDER OR CONTRIBUTORS BE LIABLE FOR ANY DIRECT,
+INDIRECT, INCIDENTAL, SPECIAL, EXEMPLARY, OR CONSEQUENTIAL DAMAGES (INCLUDING,
+BUT NOT LIMITED TO, PROCUREMENT OF SUBSTITUTE GOODS OR SERVICES; LOSS OF USE, DATA,
+OR PROFITS; OR BUSINESS INTERRUPTION) HOWEVER CAUSED AND ON ANY THEORY OF LIABILITY,
+WHETHER IN CONTRACT, STRICT LIABILITY, OR TORT (INCLUDING NEGLIGENCE OR OTHERWISE)
+ARISING IN ANY WAY OUT OF THE USE OF THIS SOFTWARE, EVEN IF ADVISED OF THE
+POSSIBILITY OF SUCH DAMAGE.
+*/
+
+
+#pragma once
+
+#include "AnimGraphNode_Base.h"
+#include "AnimNode_SteamVRInputAnimPose.h"
+
+#include "AnimGraphNode_SteamVRInputAnimPose.generated.h"
+
+using namespace vr;
+
+UCLASS()
+class UAnimGraphNode_SteamVRInputAnimPose : public UAnimGraphNode_Base
+{
+	GENERATED_UCLASS_BODY()
+
+	UPROPERTY(EditAnywhere, Category=Settings)
+	FAnimNode_SteamVRInputAnimPose Node;
+
+public:
+	// UEdGraphNode interface
+	virtual FText GetNodeTitle(ENodeTitleType::Type TitleType) const override;
+	virtual FLinearColor GetNodeTitleColor() const override;
+	virtual FText GetMenuCategory() const override;
+	virtual FText GetTooltipText() const override;
+	// End of UEdGraphNode interface
+};
\ No newline at end of file
diff --git a/Engine/Plugins/Runtime/Steam/SteamVR/Source/SteamVREditor/Public/AnimGraphNode_SteamVRSetWristTransform.h b/Engine/Plugins/Runtime/Steam/SteamVR/Source/SteamVREditor/Public/AnimGraphNode_SteamVRSetWristTransform.h
new file mode 100644
--- /dev/null
+++ b/Engine/Plugins/Runtime/Steam/SteamVR/Source/SteamVREditor/Public/AnimGraphNode_SteamVRSetWristTransform.h
@@ -0,0 +1,55 @@
+/*
+Copyright 2019 Valve Corporation under https://opensource.org/licenses/BSD-3-Clause
+
+Redistribution and use in source and binary forms, with or without modification,
+are permitted provided that the following conditions are met:
+
+1. Redistributions of source code must retain the above copyright notice, this
+   list of conditions and the following disclaimer.
+
+2. Redistributions in binary form must reproduce the above copyright notice,
+   this list of conditions and the following disclaimer in the documentation
+   and/or other materials provided with the distribution.
+
+3. Neither the name of the copyright holder nor the names of its contributors
+   may be used to endorse or promote products derived from this software
+   without specific prior written permission.
+
+THIS SOFTWARE IS PROVIDED BY THE COPYRIGHT HOLDERS AND CONTRIBUTORS "AS IS" AND
+ANY EXPRESS OR IMPLIED WARRANTIES, INCLUDING, BUT NOT LIMITED TO, THE IMPLIED
+WARRANTIES OF MERCHANTABILITY AND FITNESS FOR A PARTICULAR PURPOSE ARE DISCLAIMED.
+IN NO EVENT SHALL THE COPYRIGHT HOLDER OR CONTRIBUTORS BE LIABLE FOR ANY DIRECT,
+INDIRECT, INCIDENTAL, SPECIAL, EXEMPLARY, OR CONSEQUENTIAL DAMAGES (INCLUDING,
+BUT NOT LIMITED TO, PROCUREMENT OF SUBSTITUTE GOODS OR SERVICES; LOSS OF USE, DATA,
+OR PROFITS; OR BUSINESS INTERRUPTION) HOWEVER CAUSED AND ON ANY THEORY OF LIABILITY,
+WHETHER IN CONTRACT, STRICT LIABILITY, OR TORT (INCLUDING NEGLIGENCE OR OTHERWISE)
+ARISING IN ANY WAY OUT OF THE USE OF THIS SOFTWARE, EVEN IF ADVISED OF THE
+POSSIBILITY OF SUCH DAMAGE.
+*/
+
+
+#pragma once
+
+#include "AnimGraphNode_Base.h"
+#include "AnimNode_SteamVRSetWristTransform.h"
+
+#include "AnimGraphNode_SteamVRSetWristTransform.generated.h"
+
+using namespace vr;
+
+UCLASS()
+class UAnimGraphNode_SteamVRSetWristTransform : public UAnimGraphNode_Base
+{
+	GENERATED_UCLASS_BODY()
+
+	UPROPERTY(EditAnywhere, Category=Settings)
+	FAnimNode_SteamVRSetWristTransform Node;
+
+public:
+	// UEdGraphNode interface
+	virtual FText GetNodeTitle(ENodeTitleType::Type TitleType) const override;
+	virtual FLinearColor GetNodeTitleColor() const override;
+	virtual FText GetMenuCategory() const override;
+	virtual FText GetTooltipText() const override;
+	// End of UEdGraphNode interface
+};
\ No newline at end of file
diff --git a/Engine/Plugins/Runtime/Steam/SteamVR/Source/SteamVREditor/Public/SteamVREditor.h b/Engine/Plugins/Runtime/Steam/SteamVR/Source/SteamVREditor/Public/SteamVREditor.h
new file mode 100644
--- /dev/null
+++ b/Engine/Plugins/Runtime/Steam/SteamVR/Source/SteamVREditor/Public/SteamVREditor.h
@@ -0,0 +1,68 @@
+/*
+Copyright 2019 Valve Corporation under https://opensource.org/licenses/BSD-3-Clause
+
+Redistribution and use in source and binary forms, with or without modification,
+are permitted provided that the following conditions are met:
+
+1. Redistributions of source code must retain the above copyright notice, this
+   list of conditions and the following disclaimer.
+
+2. Redistributions in binary form must reproduce the above copyright notice,
+   this list of conditions and the following disclaimer in the documentation
+   and/or other materials provided with the distribution.
+
+3. Neither the name of the copyright holder nor the names of its contributors
+   may be used to endorse or promote products derived from this software
+   without specific prior written permission.
+
+THIS SOFTWARE IS PROVIDED BY THE COPYRIGHT HOLDERS AND CONTRIBUTORS "AS IS" AND
+ANY EXPRESS OR IMPLIED WARRANTIES, INCLUDING, BUT NOT LIMITED TO, THE IMPLIED
+WARRANTIES OF MERCHANTABILITY AND FITNESS FOR A PARTICULAR PURPOSE ARE DISCLAIMED.
+IN NO EVENT SHALL THE COPYRIGHT HOLDER OR CONTRIBUTORS BE LIABLE FOR ANY DIRECT,
+INDIRECT, INCIDENTAL, SPECIAL, EXEMPLARY, OR CONSEQUENTIAL DAMAGES (INCLUDING,
+BUT NOT LIMITED TO, PROCUREMENT OF SUBSTITUTE GOODS OR SERVICES; LOSS OF USE, DATA,
+OR PROFITS; OR BUSINESS INTERRUPTION) HOWEVER CAUSED AND ON ANY THEORY OF LIABILITY,
+WHETHER IN CONTRACT, STRICT LIABILITY, OR TORT (INCLUDING NEGLIGENCE OR OTHERWISE)
+ARISING IN ANY WAY OUT OF THE USE OF THIS SOFTWARE, EVEN IF ADVISED OF THE
+POSSIBILITY OF SUCH DAMAGE.
+*/
+
+
+#pragma once
+
+#include "Modules/ModuleManager.h"
+#include "../../SteamVRInputDevice/Public/SteamVRInputDevice.h"
+#include "../../SteamVRInputDevice/Public/SteamVRInputDeviceFunctionLibrary.h"
+
+class FToolBarBuilder;
+class FMenuBuilder;
+class SWidget;
+
+class FSteamVREditorModule : public IModuleInterface
+{
+public:
+
+	/** IModuleInterface implementation */
+	virtual void StartupModule() override;
+	virtual void ShutdownModule() override;
+	
+	/** These functions will be bound to Commands */
+	void PluginButtonClicked();
+	
+	void JsonRegenerateActionManifest();
+	void JsonRegenerateControllerBindings();
+	void ReloadActionManifest();
+	void LaunchBindingsURL();
+	void AddSampleInputs();
+
+private:
+
+	void AddToolbarExtension(FToolBarBuilder& Builder);
+	void AddMenuExtension(FMenuBuilder& Builder);
+	bool AddUniqueAxisMapping(TArray<FInputAxisKeyMapping> ExistingAxisKeys, UInputSettings* InputSettings, FName ActionName, FKey ActionKey);
+	bool AddUniqueActionMapping(TArray<FInputActionKeyMapping> ExistingActionKeys, UInputSettings* InputSettings, FName ActionName, FKey ActionKey);
+
+	TSharedPtr<class FUICommandList> PluginCommands;
+	TSharedRef<SWidget> FillComboButton(TSharedPtr<class FUICommandList> Commands);
+
+};
diff --git a/Engine/Plugins/Runtime/Steam/SteamVR/Source/SteamVREditor/SteamVREditor.Build.cs b/Engine/Plugins/Runtime/Steam/SteamVR/Source/SteamVREditor/SteamVREditor.Build.cs
new file mode 100644
--- /dev/null
+++ b/Engine/Plugins/Runtime/Steam/SteamVR/Source/SteamVREditor/SteamVREditor.Build.cs
@@ -0,0 +1,48 @@
+using UnrealBuildTool;
+using System.IO;
+
+public class SteamVREditor : ModuleRules
+{
+	public SteamVREditor(ReadOnlyTargetRules Target) : base(Target)
+    {
+        PrivatePCHHeaderFile = "Public/SteamVREditor.h";
+
+        PublicIncludePaths.Add(Path.Combine(ModuleDirectory, "Public"));
+        PrivateIncludePaths.Add(Path.Combine(ModuleDirectory, "Private"));
+
+        PublicDependencyModuleNames.AddRange(
+			new string[]
+			{
+				"Core",
+                "InputCore",
+                "ApplicationCore",
+                "InputDevice",
+                "SteamVRInput",
+                "SteamVRInputDevice",
+                "BlueprintGraph",
+                "AnimGraph"
+            }
+			);
+			
+		PrivateDependencyModuleNames.AddRange(
+			new string[]
+			{
+				"Projects",
+				"InputCore",
+				"UnrealEd",
+				"LevelEditor",
+				"CoreUObject",
+				"Engine",
+				"Slate",
+				"SlateCore"
+            }
+			);
+		
+		
+		DynamicallyLoadedModuleNames.AddRange(
+			new string[]
+			{
+			}
+			);
+	}
+}
diff --git a/Engine/Plugins/Runtime/Steam/SteamVR/Source/SteamVRInput/Private/SteamVRInput.cpp b/Engine/Plugins/Runtime/Steam/SteamVR/Source/SteamVRInput/Private/SteamVRInput.cpp
new file mode 100644
--- /dev/null
+++ b/Engine/Plugins/Runtime/Steam/SteamVR/Source/SteamVRInput/Private/SteamVRInput.cpp
@@ -0,0 +1,52 @@
+/*
+Copyright 2019 Valve Corporation under https://opensource.org/licenses/BSD-3-Clause
+
+Redistribution and use in source and binary forms, with or without modification,
+are permitted provided that the following conditions are met:
+
+1. Redistributions of source code must retain the above copyright notice, this
+   list of conditions and the following disclaimer.
+
+2. Redistributions in binary form must reproduce the above copyright notice,
+   this list of conditions and the following disclaimer in the documentation
+   and/or other materials provided with the distribution.
+
+3. Neither the name of the copyright holder nor the names of its contributors
+   may be used to endorse or promote products derived from this software
+   without specific prior written permission.
+
+THIS SOFTWARE IS PROVIDED BY THE COPYRIGHT HOLDERS AND CONTRIBUTORS "AS IS" AND
+ANY EXPRESS OR IMPLIED WARRANTIES, INCLUDING, BUT NOT LIMITED TO, THE IMPLIED
+WARRANTIES OF MERCHANTABILITY AND FITNESS FOR A PARTICULAR PURPOSE ARE DISCLAIMED.
+IN NO EVENT SHALL THE COPYRIGHT HOLDER OR CONTRIBUTORS BE LIABLE FOR ANY DIRECT,
+INDIRECT, INCIDENTAL, SPECIAL, EXEMPLARY, OR CONSEQUENTIAL DAMAGES (INCLUDING,
+BUT NOT LIMITED TO, PROCUREMENT OF SUBSTITUTE GOODS OR SERVICES; LOSS OF USE, DATA,
+OR PROFITS; OR BUSINESS INTERRUPTION) HOWEVER CAUSED AND ON ANY THEORY OF LIABILITY,
+WHETHER IN CONTRACT, STRICT LIABILITY, OR TORT (INCLUDING NEGLIGENCE OR OTHERWISE)
+ARISING IN ANY WAY OUT OF THE USE OF THIS SOFTWARE, EVEN IF ADVISED OF THE
+POSSIBILITY OF SUCH DAMAGE.
+*/
+
+
+#include "SteamVRInput.h"
+#include "CoreMinimal.h"
+#include "Runtime/Core/Public/Misc/Paths.h"
+#include "openvr.h"
+#include "Modules/ModuleManager.h"
+#include "Interfaces/IPluginManager.h"
+
+#define LOCTEXT_NAMESPACE "FSteamVRInputModule"
+
+DEFINE_LOG_CATEGORY_STATIC(LogSteamInput, Log, All);
+
+void FSteamVRInputModule::StartupModule()
+{
+}
+
+void FSteamVRInputModule::ShutdownModule()
+{
+}
+
+#undef LOCTEXT_NAMESPACE
+	
+IMPLEMENT_MODULE(FSteamVRInputModule, SteamVRInput)
diff --git a/Engine/Plugins/Runtime/Steam/SteamVR/Source/SteamVRInput/Public/SteamVRInput.h b/Engine/Plugins/Runtime/Steam/SteamVR/Source/SteamVRInput/Public/SteamVRInput.h
new file mode 100644
--- /dev/null
+++ b/Engine/Plugins/Runtime/Steam/SteamVR/Source/SteamVRInput/Public/SteamVRInput.h
@@ -0,0 +1,56 @@
+/*
+Copyright 2019 Valve Corporation under https://opensource.org/licenses/BSD-3-Clause
+
+Redistribution and use in source and binary forms, with or without modification,
+are permitted provided that the following conditions are met:
+
+1. Redistributions of source code must retain the above copyright notice, this
+   list of conditions and the following disclaimer.
+
+2. Redistributions in binary form must reproduce the above copyright notice,
+   this list of conditions and the following disclaimer in the documentation
+   and/or other materials provided with the distribution.
+
+3. Neither the name of the copyright holder nor the names of its contributors
+   may be used to endorse or promote products derived from this software
+   without specific prior written permission.
+
+THIS SOFTWARE IS PROVIDED BY THE COPYRIGHT HOLDERS AND CONTRIBUTORS "AS IS" AND
+ANY EXPRESS OR IMPLIED WARRANTIES, INCLUDING, BUT NOT LIMITED TO, THE IMPLIED
+WARRANTIES OF MERCHANTABILITY AND FITNESS FOR A PARTICULAR PURPOSE ARE DISCLAIMED.
+IN NO EVENT SHALL THE COPYRIGHT HOLDER OR CONTRIBUTORS BE LIABLE FOR ANY DIRECT,
+INDIRECT, INCIDENTAL, SPECIAL, EXEMPLARY, OR CONSEQUENTIAL DAMAGES (INCLUDING,
+BUT NOT LIMITED TO, PROCUREMENT OF SUBSTITUTE GOODS OR SERVICES; LOSS OF USE, DATA,
+OR PROFITS; OR BUSINESS INTERRUPTION) HOWEVER CAUSED AND ON ANY THEORY OF LIABILITY,
+WHETHER IN CONTRACT, STRICT LIABILITY, OR TORT (INCLUDING NEGLIGENCE OR OTHERWISE)
+ARISING IN ANY WAY OUT OF THE USE OF THIS SOFTWARE, EVEN IF ADVISED OF THE
+POSSIBILITY OF SUCH DAMAGE.
+*/
+
+
+#pragma once
+
+#include "Modules/ModuleManager.h"
+
+class FSteamVRInputModule : public IModuleInterface
+{
+public:
+
+	/* IModuleInterface implementation */
+	virtual void StartupModule() override;
+	virtual void ShutdownModule() override;
+
+	static inline IModuleInterface& Get()
+	{
+		return FModuleManager::LoadModuleChecked<IModuleInterface>("SteamVRInput");
+	}
+
+	static inline bool IsAvailable()
+	{
+		return FModuleManager::Get().IsModuleLoaded("SteamVRInput");
+	}
+
+private:
+	/* Handle to the OpenVR Library */
+	void*	OpenVRSDKHandle;
+};
diff --git a/Engine/Plugins/Runtime/Steam/SteamVR/Source/SteamVRInput/SteamVRInput.Build.cs b/Engine/Plugins/Runtime/Steam/SteamVR/Source/SteamVRInput/SteamVRInput.Build.cs
new file mode 100644
--- /dev/null
+++ b/Engine/Plugins/Runtime/Steam/SteamVR/Source/SteamVRInput/SteamVRInput.Build.cs
@@ -0,0 +1,23 @@
+using UnrealBuildTool;
+using System.IO;
+
+public class SteamVRInput : ModuleRules
+{
+	public SteamVRInput(ReadOnlyTargetRules Target) : base(Target)
+    {
+		PCHUsage = ModuleRules.PCHUsageMode.UseExplicitOrSharedPCHs;
+        PublicIncludePaths.Add(Path.Combine(ModuleDirectory, "Public"));
+        PrivateIncludePaths.Add(Path.Combine(ModuleDirectory, "Private"));
+        PublicDependencyModuleNames.AddRange(
+			new string[]
+			{
+				"Core",
+                "Projects",
+                "OpenVR",
+                "SteamVR"
+            }
+			);
+
+        RuntimeDependencies.Add("$(ProjectDir)/Config/SteamVRBindings/...");
+    }
+}
diff --git a/Engine/Plugins/Runtime/Steam/SteamVR/Source/SteamVRInputDevice/Private/AnimNode_SteamVRInputAnimPose.cpp b/Engine/Plugins/Runtime/Steam/SteamVR/Source/SteamVRInputDevice/Private/AnimNode_SteamVRInputAnimPose.cpp
new file mode 100644
--- /dev/null
+++ b/Engine/Plugins/Runtime/Steam/SteamVR/Source/SteamVRInputDevice/Private/AnimNode_SteamVRInputAnimPose.cpp
@@ -0,0 +1,339 @@
+/*
+Copyright 2019 Valve Corporation under https://opensource.org/licenses/BSD-3-Clause
+
+Redistribution and use in source and binary forms, with or without modification,
+are permitted provided that the following conditions are met:
+
+1. Redistributions of source code must retain the above copyright notice, this
+   list of conditions and the following disclaimer.
+
+2. Redistributions in binary form must reproduce the above copyright notice,
+   this list of conditions and the following disclaimer in the documentation
+   and/or other materials provided with the distribution.
+
+3. Neither the name of the copyright holder nor the names of its contributors
+   may be used to endorse or promote products derived from this software
+   without specific prior written permission.
+
+THIS SOFTWARE IS PROVIDED BY THE COPYRIGHT HOLDERS AND CONTRIBUTORS "AS IS" AND
+ANY EXPRESS OR IMPLIED WARRANTIES, INCLUDING, BUT NOT LIMITED TO, THE IMPLIED
+WARRANTIES OF MERCHANTABILITY AND FITNESS FOR A PARTICULAR PURPOSE ARE DISCLAIMED.
+IN NO EVENT SHALL THE COPYRIGHT HOLDER OR CONTRIBUTORS BE LIABLE FOR ANY DIRECT,
+INDIRECT, INCIDENTAL, SPECIAL, EXEMPLARY, OR CONSEQUENTIAL DAMAGES (INCLUDING,
+BUT NOT LIMITED TO, PROCUREMENT OF SUBSTITUTE GOODS OR SERVICES; LOSS OF USE, DATA,
+OR PROFITS; OR BUSINESS INTERRUPTION) HOWEVER CAUSED AND ON ANY THEORY OF LIABILITY,
+WHETHER IN CONTRACT, STRICT LIABILITY, OR TORT (INCLUDING NEGLIGENCE OR OTHERWISE)
+ARISING IN ANY WAY OUT OF THE USE OF THIS SOFTWARE, EVEN IF ADVISED OF THE
+POSSIBILITY OF SUCH DAMAGE.
+*/
+
+
+#include "AnimNode_SteamVRInputAnimPose.h"
+#include "ISteamVRInputDeviceModule.h"
+#include "AnimationRuntime.h"
+#include "Runtime/Engine/Public/Animation/AnimInstanceProxy.h"
+#include "SteamVRInputDevice.h"
+#include "UE4HandSkeletonDefinition.h"
+
+// Lookup table that maps each bone in the UE4 hand skeleton to the corresponding
+// bone in the SteamVR hand skeleton
+static const int32 kUE4BoneToSteamVRBone[] = {
+	ESteamVRBone_Wrist,			// EUE4HandBone_Wrist
+	ESteamVRBone_IndexFinger1,	// EUE4HandBone_Index_01
+	ESteamVRBone_IndexFinger2,	// EUE4HandBone_Index_02
+	ESteamVRBone_IndexFinger3,	// EUE4HandBone_Index_03
+	ESteamVRBone_MiddleFinger1, // EUE4HandBone_Middle_01
+	ESteamVRBone_MiddleFinger2, // EUE4HandBone_Middle_02
+	ESteamVRBone_MiddleFinger3, // EUE4HandBone_Middle_03
+	ESteamVRBone_PinkyFinger1,	// EUE4HandBone_Pinky_01
+	ESteamVRBone_PinkyFinger2,	// EUE4HandBone_Pinky_02
+	ESteamVRBone_PinkyFinger3,	// EUE4HandBone_Pinky_03
+	ESteamVRBone_RingFinger1,	// EUE4HandBone_Ring_01
+	ESteamVRBone_RingFinger2,	// EUE4HandBone_Ring_02
+	ESteamVRBone_RingFinger3,	// EUE4HandBone_Ring_03
+	ESteamVRBone_Thumb0,		// EUE4HandBone_Thumb_01
+	ESteamVRBone_Thumb1,		// EUE4HandBone_Thumb_02
+	ESteamVRBone_Thumb2			// EUE4HandBone_Thumb_03
+};
+static_assert( sizeof(kUE4BoneToSteamVRBone) / sizeof(kUE4BoneToSteamVRBone[0]) == EUE4HandBone_Count, "Mapping from UE4 hand bones to their corresponding SteamVR hand bones is the wrong size" );
+
+// List of the knuckle bones of the SteamVR hand skeleton
+static const int32 kSteamVRKnuckleBones[] = {
+	ESteamVRBone_IndexFinger1,
+	ESteamVRBone_MiddleFinger1,
+	ESteamVRBone_RingFinger1,
+	ESteamVRBone_PinkyFinger1
+};
+static_assert( sizeof(kSteamVRKnuckleBones) / sizeof(kSteamVRKnuckleBones[0]) == 4, "List of SteamVR knuckle bones should only have 4 entries" );
+
+// List of the knuckle bones of the UE4 hand skeleton
+static const int32 kUE4KnuckleBones[] = {
+	EUE4HandBone_Index_01,
+	EUE4HandBone_Middle_01,
+	EUE4HandBone_Ring_01,
+	EUE4HandBone_Pinky_01
+};
+static_assert( sizeof(kUE4KnuckleBones) / sizeof(kUE4KnuckleBones[0]) == 4, "List of UE4 knuckle bones should only have 4 entries" );
+
+
+FAnimNode_SteamVRInputAnimPose::FAnimNode_SteamVRInputAnimPose()
+{
+}
+
+void FAnimNode_SteamVRInputAnimPose::Initialize_AnyThread(const FAnimationInitializeContext& Context)
+{
+}
+
+void FAnimNode_SteamVRInputAnimPose::CacheBones_AnyThread(const FAnimationCacheBonesContext & Context)
+{
+}
+
+void FAnimNode_SteamVRInputAnimPose::Update_AnyThread(const FAnimationUpdateContext & Context)
+{
+	// Grab node inputs
+	GetEvaluateGraphExposedInputs().Execute(Context);
+}
+
+void FAnimNode_SteamVRInputAnimPose::Evaluate_AnyThread(FPoseContext& Output)
+{
+	Output.ResetToRefPose();
+
+	// Update Skeletal Animation
+	FSteamVRInputDevice* SteamVRInputDevice = GetSteamVRInputDevice();
+	if (SteamVRInputDevice != nullptr)
+	{
+		// Attempt to read the current skeletal pose from SteamVR.  The data returned will be the bone transforms of the SteamVR hand skeleton,
+		// transformed into the UE4 coordinate system
+
+		FTransform BoneTransforms[STEAMVR_SKELETON_BONE_COUNT];
+		EVRSkeletalMotionRange SteamVRMotionRange = (MotionRange == EMotionRange::VR_WithController) ? VRSkeletalMotionRange_WithController : VRSkeletalMotionRange_WithoutController;
+		bool bIsLeftHand = (Hand == EHand::VR_LeftHand);
+		
+		if (SteamVRInputDevice->GetSkeletalData(bIsLeftHand, Mirror, SteamVRMotionRange, BoneTransforms, STEAMVR_SKELETON_BONE_COUNT))
+		{
+			// If the target hand skeleton is the SteamVR skeleton, then we can just copy the transforms directly into the pose
+			if (HandSkeleton == EHandSkeleton::VR_SteamVRHandSkeleton)
+			{
+				for (int32 SrcBoneindex = 0; SrcBoneindex < STEAMVR_SKELETON_BONE_COUNT; ++SrcBoneindex)
+				{
+					FCompactPoseBoneIndex TargetBoneIndex = Output.Pose.GetBoneContainer().MakeCompactPoseIndex(FMeshPoseBoneIndex(SrcBoneindex));
+					if (TargetBoneIndex != INDEX_NONE)
+					{
+						Output.Pose[TargetBoneIndex] = BoneTransforms[SrcBoneindex];
+					}
+				}
+			}
+			else
+			{
+				// If the target hand skeleton is the UE4 reference VR hand skeleton, we need to retarget the transforms
+				// we read from SteamVR to fit it.  
+				PoseUE4HandSkeleton(Output.Pose, BoneTransforms, STEAMVR_SKELETON_BONE_COUNT);
+			}
+		}
+	}
+}
+
+
+FQuat CalcRotationAboutAxis(const FVector& FromDirection, const FVector& ToDirection, const FVector& Axis)
+{
+	FVector FromDirectionCp = FVector::CrossProduct(Axis, FromDirection);
+	FVector ToDirectionCp = FVector::CrossProduct(Axis, ToDirection);
+
+	return FQuat::FindBetweenVectors(FromDirectionCp, ToDirectionCp);
+}
+
+
+void FAnimNode_SteamVRInputAnimPose::PoseUE4HandSkeleton(FCompactPose& Pose, const FTransform* BoneTransformsLS, int32 BoneTransformCount)
+{
+	check(BoneTransformsLS != nullptr);
+	check(BoneTransformCount == SteamVRSkeleton::GetBoneCount());
+	check(Pose.GetNumBones() == UE4HandSkeleton::GetBoneCount());
+
+	// It is easier to do the retargetting in model space, so calculate the model space transforms for the SteamVR skeleton
+	// from the given local space transforms
+	FTransform BoneTransformsMS[STEAMVR_SKELETON_BONE_COUNT];
+	for ( int32 BoneIndex = 0; BoneIndex < SteamVRSkeleton::GetBoneCount(); ++BoneIndex)
+	{
+		int32 ParentIndex = SteamVRSkeleton::GetParentIndex(BoneIndex);
+		if (ParentIndex != -1)
+		{
+			BoneTransformsMS[BoneIndex] = BoneTransformsLS[BoneIndex] * BoneTransformsMS[ParentIndex];
+		}
+		else
+		{
+			BoneTransformsMS[BoneIndex] = BoneTransformsLS[BoneIndex];
+		}
+	}
+
+	// Remove any scale, as its now baked into the position
+	for (int32 BoneIndex = 0; BoneIndex < SteamVRSkeleton::GetBoneCount(); ++BoneIndex)
+	{
+		BoneTransformsMS[BoneIndex].SetScale3D(FVector::ZeroVector);
+	}
+
+	FQuat TargetBoneRotationsMS[EUE4HandBone_Count];
+
+	// Cache UE4 retargetting references at first anim frame as FAnimPoseBase init does not appear to have a valid Pose to process at Initialize()
+	if (!UE4RetargettingRefs.bIsInitialized)
+	{
+		// Determine which hand we're working with
+		UE4RetargettingRefs.bIsRightHanded = (Hand == EHand::VR_RightHand && !Mirror) || (Hand == EHand::VR_LeftHand && Mirror);
+
+		// Calculate the average position of the UE4 knuckles bones and add it to the cache
+		for (int32 KnuckleIndex = 0; KnuckleIndex < 4; ++KnuckleIndex)
+		{
+			const FCompactPoseBoneIndex BoneIndex_UE4 = FCompactPoseBoneIndex(kUE4KnuckleBones[KnuckleIndex]);
+			UE4RetargettingRefs.KnuckleAverageMS_UE4 += CalcModelSpaceTransform(Pose, BoneIndex_UE4).GetTranslation();
+		}
+
+		UE4RetargettingRefs.KnuckleAverageMS_UE4 /= 4.f;
+
+		// Obtain the UE4 wrist Side & Forward directions from first animation frame and place in cache 
+		FCompactPoseBoneIndex WristBoneIndexCompact = Pose.GetBoneContainer().MakeCompactPoseIndex(FMeshPoseBoneIndex(EUE4HandBone_Wrist));
+		FTransform WristTransform_UE4 = Pose[WristBoneIndexCompact];
+		FVector ToKnuckleAverageMS_UE4 = UE4RetargettingRefs.KnuckleAverageMS_UE4 - WristTransform_UE4.GetTranslation();
+		ToKnuckleAverageMS_UE4.Normalize();
+
+		UE4RetargettingRefs.WristForwardLS_UE4 = WristTransform_UE4.GetRotation().UnrotateVector(ToKnuckleAverageMS_UE4);
+		UE4RetargettingRefs.WristSideDirectionLS_UE4 = FVector::CrossProduct(UE4RetargettingRefs.WristForwardLS_UE4, FVector::RightVector);
+
+		// Set initialized flag to true so that subsequent frames will skip the above calculations
+		UE4RetargettingRefs.bIsInitialized = true;
+	}
+
+	// Calculate the rotation to make the target wrist bone match the orientation of the source wrist bone
+	{
+		// Calculate the average position of the knuckles bones for each hand, then rotate the wrist of the target skeleton
+		// so that its knuckle-average coincides with the knuckle-average of the source pose
+		FVector KnuckleAverageMS_SteamVR = FVector::ZeroVector;
+
+		for (int32 KnuckleIndex = 0; KnuckleIndex < 4; ++KnuckleIndex)
+		{
+			const int32 BoneIndex_SteamVR = kSteamVRKnuckleBones[KnuckleIndex];
+			KnuckleAverageMS_SteamVR	+= BoneTransformsMS[BoneIndex_SteamVR].GetTranslation();
+		}
+
+		KnuckleAverageMS_SteamVR /= 4.f;
+
+		FVector ToKnuckleAverageMS_SteamVR = KnuckleAverageMS_SteamVR - BoneTransformsMS[ESteamVRBone_Wrist].GetTranslation();
+		ToKnuckleAverageMS_SteamVR.Normalize();
+		FVector WristForwardLS_SteamVR = BoneTransformsMS[ESteamVRBone_Wrist].GetRotation().UnrotateVector(ToKnuckleAverageMS_SteamVR);
+
+		// Get the axis that most closely matches the direction the palm of the hand is facing
+		FVector WristSideDirectionLS_SteamVR = (UE4RetargettingRefs.bIsRightHanded) ? FVector(0.f, -1.f, 0.f) : FVector(0.f, 1.f, 0.f);
+
+		// Take the cross product with the forward vector for the SteamVR hand to ensure that the side vector is perpendicular to the forward vector
+		WristSideDirectionLS_SteamVR = FVector::CrossProduct(WristForwardLS_SteamVR, WristSideDirectionLS_SteamVR);
+
+		// Find the model-space directions of the forward and side vectors based on the current pose
+		const FVector WristForwardMS_SteamVR = BoneTransformsMS[ESteamVRBone_Wrist].GetRotation() * WristForwardLS_SteamVR;
+		const FVector WristSideDirectionMS_SteamVR = BoneTransformsMS[ESteamVRBone_Wrist].GetRotation() * WristSideDirectionLS_SteamVR;
+
+		// Calculate the rotation that will align the UE4 hand's forward vector with the SteamVR hand's forward
+		FQuat AlignmentRot = FQuat::FindBetweenNormals(UE4RetargettingRefs.WristForwardLS_UE4, WristForwardMS_SteamVR);
+
+		// Rotate about the aligned forward direction to make the side directions align
+		FVector WristSideDirectionMS_UE4 = AlignmentRot * UE4RetargettingRefs.WristSideDirectionLS_UE4;
+		FQuat TwistRotation = CalcRotationAboutAxis(WristSideDirectionMS_UE4, WristSideDirectionMS_SteamVR, WristForwardMS_SteamVR);
+
+		// Apply the rotation to the hand
+		TargetBoneRotationsMS[EUE4HandBone_Wrist] = TwistRotation * AlignmentRot;
+	}
+
+	// For all the remaining bones, use their child bone as a reference to calculate their orientation
+	const FVector FingerForwardDefault_SteamVR(-1.f, 0.f, 0.f);
+	const FVector FingerForwardDefault_UE4(-1.f, 0.f, 0.f);
+	for ( int32 BoneIndex = EUE4HandBone_Wrist+1; BoneIndex < Pose.GetNumBones(); ++BoneIndex)
+	{
+		// Determine which direction is "forward" in the bone's local space by looking at the direction
+		// to its child
+		FVector FingerForwardLS_UE4 = FingerForwardDefault_UE4;
+		if (UE4HandSkeleton::GetChildCount(BoneIndex) > 0)
+		{
+			int32 ChildBoneIndex = UE4HandSkeleton::GetChildIndex(BoneIndex, 0);
+			FCompactPoseBoneIndex ChildBoneIndexCompact = Pose.GetBoneContainer().MakeCompactPoseIndex(FMeshPoseBoneIndex(ChildBoneIndex));
+			FingerForwardLS_UE4 = Pose[ChildBoneIndexCompact].GetTranslation();
+			FingerForwardLS_UE4.Normalize();
+		}
+
+		int32 ParentBoneIndex = UE4HandSkeleton::GetParentIndex(BoneIndex);
+		check(ParentBoneIndex != -1);
+
+		FQuat StartingTransformMS = TargetBoneRotationsMS[ParentBoneIndex];
+
+		// Include the default orientation of the bone, so that it is used as the starting point for the adjustment rotation
+		FCompactPoseBoneIndex BoneIndexCompact = Pose.GetBoneContainer().MakeCompactPoseIndex(FMeshPoseBoneIndex(BoneIndex));
+		StartingTransformMS = StartingTransformMS * Pose[BoneIndexCompact].GetRotation();
+
+		// Convert the bone's forward direction to model space
+		FVector FingerForwardMS_UE4 = StartingTransformMS.RotateVector(FingerForwardLS_UE4);
+		
+		// Calculate the direction that the bone's forward vector should be pointing
+		const int32 BoneIndex_SteamVR = kUE4BoneToSteamVRBone[BoneIndex];
+		FVector FingerForwardMS_SteamVR = FingerForwardDefault_SteamVR;
+		if (SteamVRSkeleton::GetChildCount(BoneIndex_SteamVR) > 0)
+		{
+			int32 ChildIndex = SteamVRSkeleton::GetChildIndex(BoneIndex_SteamVR, 0);
+			FingerForwardMS_SteamVR = BoneTransformsMS[ChildIndex].GetTranslation() - BoneTransformsMS[BoneIndex_SteamVR].GetTranslation();
+			FingerForwardMS_SteamVR.Normalize();
+		}
+
+		// Find the rotation between the current forward vector and the desired forward vector, in model space
+		FQuat AdjustmentRot = FQuat::FindBetweenVectors(FingerForwardMS_UE4, FingerForwardMS_SteamVR);
+
+		TargetBoneRotationsMS[BoneIndex] = AdjustmentRot * StartingTransformMS;
+	}
+
+	// Convert the target rotations from model-space to local-space and apply them on the output pose
+	for (int32 BoneIndex = 0; BoneIndex < Pose.GetNumBones(); ++BoneIndex)
+	{
+		FCompactPoseBoneIndex TargetBoneIndex = Pose.GetBoneContainer().MakeCompactPoseIndex(FMeshPoseBoneIndex(BoneIndex));
+		if (TargetBoneIndex != INDEX_NONE)
+		{
+			FQuat BoneRotation = TargetBoneRotationsMS[BoneIndex];
+
+			int32 ParentIndex = UE4HandSkeleton::GetParentIndex(BoneIndex);
+			if (ParentIndex != -1)
+			{
+				BoneRotation = TargetBoneRotationsMS[ParentIndex].Inverse() * BoneRotation;
+			}
+
+			Pose[TargetBoneIndex].SetRotation(BoneRotation);
+		}
+	}
+
+	// Set the wrist bone position
+	{
+		FCompactPoseBoneIndex WristBoneIndex = Pose.GetBoneContainer().MakeCompactPoseIndex(FMeshPoseBoneIndex(EUE4HandBone_Wrist));
+		Pose[WristBoneIndex].SetTranslation(BoneTransformsMS[ESteamVRBone_Wrist].GetTranslation());
+	}
+}
+
+
+FSteamVRInputDevice* FAnimNode_SteamVRInputAnimPose::GetSteamVRInputDevice()
+{
+	TArray<IMotionController*> MotionControllers = IModularFeatures::Get().GetModularFeatureImplementations<IMotionController>(IMotionController::GetModularFeatureName());
+	for (auto MotionController : MotionControllers)
+	{
+		FSteamVRInputDevice* TestSteamVRDevice = static_cast<FSteamVRInputDevice*>(MotionController);
+		if (TestSteamVRDevice != nullptr && !FGenericPlatformMath::IsNaN(TestSteamVRDevice->DeviceSignature) && TestSteamVRDevice->DeviceSignature == 2019)
+		{
+			return TestSteamVRDevice;
+		}
+	}
+	return nullptr;
+}
+
+
+FTransform FAnimNode_SteamVRInputAnimPose::CalcModelSpaceTransform(const FCompactPose& Pose, FCompactPoseBoneIndex BoneIndex)
+{
+	FTransform BoneTransform = Pose[BoneIndex];
+
+	FCompactPoseBoneIndex ParentIndex = Pose.GetParentBoneIndex(BoneIndex);
+	if (ParentIndex != INDEX_NONE)
+	{
+		BoneTransform = BoneTransform * CalcModelSpaceTransform(Pose, ParentIndex);
+	}
+
+	return BoneTransform;
+}
diff --git a/Engine/Plugins/Runtime/Steam/SteamVR/Source/SteamVRInputDevice/Private/AnimNode_SteamVRSetWristTransform.cpp b/Engine/Plugins/Runtime/Steam/SteamVR/Source/SteamVRInputDevice/Private/AnimNode_SteamVRSetWristTransform.cpp
new file mode 100644
--- /dev/null
+++ b/Engine/Plugins/Runtime/Steam/SteamVR/Source/SteamVRInputDevice/Private/AnimNode_SteamVRSetWristTransform.cpp
@@ -0,0 +1,85 @@
+/*
+Copyright 2019 Valve Corporation under https://opensource.org/licenses/BSD-3-Clause
+
+Redistribution and use in source and binary forms, with or without modification,
+are permitted provided that the following conditions are met:
+
+1. Redistributions of source code must retain the above copyright notice, this
+   list of conditions and the following disclaimer.
+
+2. Redistributions in binary form must reproduce the above copyright notice,
+   this list of conditions and the following disclaimer in the documentation
+   and/or other materials provided with the distribution.
+
+3. Neither the name of the copyright holder nor the names of its contributors
+   may be used to endorse or promote products derived from this software
+   without specific prior written permission.
+
+THIS SOFTWARE IS PROVIDED BY THE COPYRIGHT HOLDERS AND CONTRIBUTORS "AS IS" AND
+ANY EXPRESS OR IMPLIED WARRANTIES, INCLUDING, BUT NOT LIMITED TO, THE IMPLIED
+WARRANTIES OF MERCHANTABILITY AND FITNESS FOR A PARTICULAR PURPOSE ARE DISCLAIMED.
+IN NO EVENT SHALL THE COPYRIGHT HOLDER OR CONTRIBUTORS BE LIABLE FOR ANY DIRECT,
+INDIRECT, INCIDENTAL, SPECIAL, EXEMPLARY, OR CONSEQUENTIAL DAMAGES (INCLUDING,
+BUT NOT LIMITED TO, PROCUREMENT OF SUBSTITUTE GOODS OR SERVICES; LOSS OF USE, DATA,
+OR PROFITS; OR BUSINESS INTERRUPTION) HOWEVER CAUSED AND ON ANY THEORY OF LIABILITY,
+WHETHER IN CONTRACT, STRICT LIABILITY, OR TORT (INCLUDING NEGLIGENCE OR OTHERWISE)
+ARISING IN ANY WAY OUT OF THE USE OF THIS SOFTWARE, EVEN IF ADVISED OF THE
+POSSIBILITY OF SUCH DAMAGE.
+*/
+
+
+#include "AnimNode_SteamVRSetWristTransform.h"
+#include "ISteamVRInputDeviceModule.h"
+#include "AnimationRuntime.h"
+#include "Runtime/Engine/Public/Animation/AnimInstanceProxy.h"
+#include "SteamVRInputDevice.h"
+#include "UE4HandSkeletonDefinition.h"
+
+FAnimNode_SteamVRSetWristTransform::FAnimNode_SteamVRSetWristTransform()
+{
+}
+
+void FAnimNode_SteamVRSetWristTransform::Initialize_AnyThread(const FAnimationInitializeContext& Context)
+{
+	// Initialize our poses
+	ReferencePose.Initialize(Context);
+	TargetPose.Initialize(Context);
+
+}
+
+void FAnimNode_SteamVRSetWristTransform::CacheBones_AnyThread(const FAnimationCacheBonesContext& Context)
+{
+
+}
+
+void FAnimNode_SteamVRSetWristTransform::Update_AnyThread(const FAnimationUpdateContext& Context)
+{
+	// Update our poses this frame
+	ReferencePose.Update(Context);
+	TargetPose.Update(Context);
+}
+
+void FAnimNode_SteamVRSetWristTransform::Evaluate_AnyThread(FPoseContext& Output)
+{
+	Output.ResetToRefPose();
+
+	// Apply all the bones from the target to the output pose
+	TargetPose.Evaluate(Output);
+
+	// Setup buffer for the reference pose
+	FPoseContext ReferenceContext = FPoseContext(Output.AnimInstanceProxy);
+	ReferencePose.Evaluate(ReferenceContext);
+
+	// Apply the appropriate root and/or wrist transforms to our output pose based on the skeleton being used
+	if (HandSkeleton == EHandSkeleton::VR_SteamVRHandSkeleton)
+	{
+		// Since we are dealing with the SteamVR Hand Skeleton, we need to set both the root and wrist Bones
+		Output.Pose[RootBoneIndex] = ReferenceContext.Pose[RootBoneIndex];
+		Output.Pose[SteamVRWristBoneIndex] = ReferenceContext.Pose[SteamVRWristBoneIndex];
+	}
+	else
+	{
+		// For the UE4 stock hand skeleton, we only need the root bone (which is the wrist bone in this skeleton)
+		Output.Pose[RootBoneIndex] = ReferenceContext.Pose[RootBoneIndex];
+	}
+}
