Deprecate MotionController keys

diff --git a/Engine/Plugins/Runtime/Steam/SteamVR/Source/SteamVRInputDevice/Private/SteamVRInputDevice.cpp b/Engine/Plugins/Runtime/Steam/SteamVR/Source/SteamVRInputDevice/Private/SteamVRInputDevice.cpp
new file mode 100644
--- /dev/null
+++ b/Engine/Plugins/Runtime/Steam/SteamVR/Source/SteamVRInputDevice/Private/SteamVRInputDevice.cpp
@@ -0,0 +1,3712 @@
+/*
+Copyright 2019 Valve Corporation under https://opensource.org/licenses/BSD-3-Clause
+
+Redistribution and use in source and binary forms, with or without modification,
+are permitted provided that the following conditions are met:
+
+1. Redistributions of source code must retain the above copyright notice, this
+   list of conditions and the following disclaimer.
+
+2. Redistributions in binary form must reproduce the above copyright notice,
+   this list of conditions and the following disclaimer in the documentation
+   and/or other materials provided with the distribution.
+
+3. Neither the name of the copyright holder nor the names of its contributors
+   may be used to endorse or promote products derived from this software
+   without specific prior written permission.
+
+THIS SOFTWARE IS PROVIDED BY THE COPYRIGHT HOLDERS AND CONTRIBUTORS "AS IS" AND
+ANY EXPRESS OR IMPLIED WARRANTIES, INCLUDING, BUT NOT LIMITED TO, THE IMPLIED
+WARRANTIES OF MERCHANTABILITY AND FITNESS FOR A PARTICULAR PURPOSE ARE DISCLAIMED.
+IN NO EVENT SHALL THE COPYRIGHT HOLDER OR CONTRIBUTORS BE LIABLE FOR ANY DIRECT,
+INDIRECT, INCIDENTAL, SPECIAL, EXEMPLARY, OR CONSEQUENTIAL DAMAGES (INCLUDING,
+BUT NOT LIMITED TO, PROCUREMENT OF SUBSTITUTE GOODS OR SERVICES; LOSS OF USE, DATA,
+OR PROFITS; OR BUSINESS INTERRUPTION) HOWEVER CAUSED AND ON ANY THEORY OF LIABILITY,
+WHETHER IN CONTRACT, STRICT LIABILITY, OR TORT (INCLUDING NEGLIGENCE OR OTHERWISE)
+ARISING IN ANY WAY OUT OF THE USE OF THIS SOFTWARE, EVEN IF ADVISED OF THE
+POSSIBILITY OF SUCH DAMAGE.
+*/
+
+
+#include "SteamVRInputDevice.h"
+#include "Runtime/ApplicationCore/Public/GenericPlatform/IInputInterface.h"
+#include "HAL/FileManagerGeneric.h"
+#include "Misc/FileHelper.h"
+#include "GameFramework/PlayerInput.h"
+#include "Engine/World.h"
+#include "Kismet/GameplayStatics.h"
+#include "openvr.h"
+#include "GameFramework/WorldSettings.h"
+#include "Features/IModularFeatures.h"
+#include "MotionControllerComponent.h"
+#include "IMotionController.h"
+#include "SteamVRControllerKeys.h"
+#include "Runtime/HeadMountedDisplay/Public/IXRTrackingSystem.h"
+#include "SteamVRSkeletonDefinition.h"
+
+#if PLATFORM_WINDOWS
+#include "Windows/WindowsHWrapper.h"
+#endif
+
+#if WITH_EDITOR
+#include "Editor.h"
+#endif
+
+#define LOCTEXT_NAMESPACE "SteamVRInputDevice"
+DEFINE_LOG_CATEGORY_STATIC(LogSteamVRInputDevice, Log, All);
+
+// List of bones that are effectively in model space because
+// they are children of the root
+static const int32 kModelSpaceBones[] = {
+	ESteamVRBone_Wrist,
+	ESteamVRBone_Aux_Thumb,
+	ESteamVRBone_Aux_IndexFinger,
+	ESteamVRBone_Aux_MiddleFinger,
+	ESteamVRBone_Aux_RingFinger,
+	ESteamVRBone_Aux_PinkyFinger,
+};
+
+// List of the metacarpal bones of the SteamVR skeleton
+static const int32 kMetacarpalBones[] = {
+	ESteamVRBone_Thumb0,
+	ESteamVRBone_IndexFinger0,
+	ESteamVRBone_MiddleFinger0,
+	ESteamVRBone_RingFinger0,
+	ESteamVRBone_PinkyFinger0
+};
+
+// List of bones that only need to have their translation mirrored in the SteamVR skeleton
+static const int32 kMirrorTranslationOnlyBones[] = {
+	ESteamVRBone_Thumb1,
+	ESteamVRBone_Thumb2,
+	ESteamVRBone_Thumb3,
+	ESteamVRBone_IndexFinger1,
+	ESteamVRBone_IndexFinger2,
+	ESteamVRBone_IndexFinger3,
+	ESteamVRBone_IndexFinger4,
+	ESteamVRBone_MiddleFinger1,
+	ESteamVRBone_MiddleFinger2,
+	ESteamVRBone_MiddleFinger3,
+	ESteamVRBone_MiddleFinger4,
+	ESteamVRBone_RingFinger1,
+	ESteamVRBone_RingFinger2,
+	ESteamVRBone_RingFinger3,
+	ESteamVRBone_RingFinger4,
+	ESteamVRBone_PinkyFinger1,
+	ESteamVRBone_PinkyFinger2,
+	ESteamVRBone_PinkyFinger3,
+	ESteamVRBone_PinkyFinger4
+};
+
+
+FSteamVRInputDevice::FSteamVRInputDevice(const TSharedRef<FGenericApplicationMessageHandler>& InMessageHandler)
+	: MessageHandler(InMessageHandler)
+{
+	// Initializations
+	InitSteamVRSystem();
+	InitControllerMappings();
+	InitControllerKeys();
+
+#if WITH_EDITOR
+	GenerateActionManifest();
+
+	// Auto-enable SteamVR Input Developer Mode (reload hmd module)
+	if (VRSettings() != nullptr)
+	{
+		EVRInitError SteamVRInitError = VRInitError_Driver_NotLoaded;
+		IVRSystem* SteamVRSystem = vr::VR_Init(&SteamVRInitError, vr::VRApplication_Overlay);
+
+		EVRSettingsError BindingFlagError = VRSettingsError_None;
+		VRSettings()->SetBool(k_pch_SteamVR_Section, k_pch_SteamVR_DebugInputBinding, true, &BindingFlagError);
+		UE_LOG(LogSteamVRInputDevice, Display, TEXT("[STEAMVR INPUT] Enable SteamVR Input Developer Mode: %s"), *FString(UTF8_TO_TCHAR(VRSettings()->GetSettingsErrorNameFromEnum(BindingFlagError))));
+		VRSettings()->SetBool(k_pch_SteamVR_Section, k_pch_SteamVR_DebugInput, true, &BindingFlagError);
+		UE_LOG(LogSteamVRInputDevice, Display, TEXT("[STEAMVR INPUT] Enable SteamVR Debug Input: %s"), *FString(UTF8_TO_TCHAR(VRSettings()->GetSettingsErrorNameFromEnum(BindingFlagError))));
+
+		VR_Shutdown();
+	}
+#endif
+
+	IModularFeatures::Get().RegisterModularFeature(GetModularFeatureName(), this);
+}
+
+FSteamVRInputDevice::~FSteamVRInputDevice()
+{
+	IModularFeatures::Get().UnregisterModularFeature(GetModularFeatureName(), this);
+}
+
+void FSteamVRInputDevice::InitSteamVRSystem()
+{
+	//UE_LOG(LogTemp, Warning, TEXT("Attempting to load Steam VR System..."));
+
+	if (!VRSystem())
+	{
+		// Initialize OpenVR
+		EVRInitError SteamVRInitError = VRInitError_Driver_NotLoaded;
+		IVRSystem* SteamVRSystem = vr::VR_Init(&SteamVRInitError, vr::VRApplication_Scene);
+
+		if (SteamVRInitError != VRInitError_None)
+		{
+			return;
+		}
+	}
+	else if (VRSystem() && VRInput() && IsInGameThread())
+	{
+		UE_LOG(LogSteamVRInputDevice, Display, TEXT("SteamVR runtime %u.%u.%u loaded."), k_nSteamVRVersionMajor, k_nSteamVRVersionMinor, k_nSteamVRVersionBuild);
+		
+		// Set Skeletal Handles
+		bIsSkeletalControllerLeftPresent = SetSkeletalHandle(TCHAR_TO_UTF8(*FString(TEXT(ACTION_PATH_SKELETON_LEFT))), VRSkeletalHandleLeft);
+		bIsSkeletalControllerRightPresent = SetSkeletalHandle(TCHAR_TO_UTF8(*FString(TEXT(ACTION_PATH_SKELETON_RIGHT))), VRSkeletalHandleRight);
+
+		// (Re)Load Action Manifest
+		GenerateActionManifest(false, false, true, false);
+
+		// Set haptic handles
+		LastInputError = VRInput()->GetActionHandle(TCHAR_TO_UTF8(*FString(TEXT(ACTION_PATH_VIBRATE_LEFT))), &VRVibrationLeft);
+		if (LastInputError != VRInputError_None || VRVibrationLeft == k_ulInvalidActionHandle)
+		{
+			VRVibrationLeft = k_ulInvalidActionHandle;
+		}
+
+		LastInputError = VRInput()->GetActionHandle(TCHAR_TO_UTF8(*FString(TEXT(ACTION_PATH_VIBRATE_RIGHT))), &VRVibrationRight);
+		if (LastInputError != VRInputError_None || VRVibrationRight == k_ulInvalidActionHandle)
+		{
+			VRVibrationRight = k_ulInvalidActionHandle;
+		}
+	}
+}
+
+void FSteamVRInputDevice::Tick(float DeltaTime)
+{
+	// Place current delta time in buffer for use in determining haptic duration
+	CurrentDeltaTime = DeltaTime;
+
+	// Watch for SteamVR availability & restarts
+	if (!VRSystem())
+	{
+		InitSteamVRSystem();
+	}
+
+	// Cache the controller transform to ensure ResetOrientationAndPosition gets the correct values (Valid for UE4.18 upwards)
+	// [REDACTED_PROFILE_URL]/steamvr_unreal_plugin/issues/2
+	if (GEngine->XRSystem.IsValid())
+	{
+		CachedBaseOrientation = GEngine->XRSystem->GetBaseOrientation();
+		CachedBasePosition = GEngine->XRSystem->GetBasePosition();
+	}
+	else
+	{
+		CachedBaseOrientation = FQuat::Identity;
+		CachedBasePosition = FVector::ZeroVector;
+	}
+}
+
+void FSteamVRInputDevice::FindAxisMappings(const UInputSettings* InputSettings, const FName InAxisName, TArray<FInputAxisKeyMapping>& OutMappings) const
+{
+	if (InAxisName.IsValid())
+	{
+		for (int32 AxisIndex = InputSettings->GetAxisMappings().Num() - 1; AxisIndex >= 0; --AxisIndex)
+		{
+			if (InputSettings->GetAxisMappings()[AxisIndex].AxisName == InAxisName)
+			{
+				OutMappings.Add(InputSettings->GetAxisMappings()[AxisIndex]);
+			}
+		}
+	}
+}
+
+void FSteamVRInputDevice::GetSteamVRMappings(TArray<FInputAxisKeyMapping>& InUEKeyMappings, TArray<FSteamVRAxisKeyMapping>& OutMappings) 
+{
+	OutMappings.Empty();
+	for (auto& UEKeyMapping : InUEKeyMappings)
+	{
+		OutMappings.Add(FSteamVRAxisKeyMapping(UEKeyMapping, false, false));
+	}
+}
+
+void FSteamVRInputDevice::FindActionMappings(const UInputSettings* InputSettings, const FName InActionName, TArray<FInputActionKeyMapping>& OutMappings) const
+{
+	if (InActionName.IsValid())
+	{
+		for (int32 ActionIndex = InputSettings->GetActionMappings().Num() - 1; ActionIndex >= 0; --ActionIndex)
+		{
+			if (InputSettings->GetActionMappings()[ActionIndex].ActionName == InActionName)
+			{
+				OutMappings.Add(InputSettings->GetActionMappings()[ActionIndex]);
+			}
+		}
+	}
+}
+
+FString FSteamVRInputDevice::SanitizeString(FString& InString)
+{
+	FString SanitizedString = InString.Replace(TEXT(" "), TEXT("-"));
+	SanitizedString = SanitizedString.Replace(TEXT("*"), TEXT("-"));
+	SanitizedString = SanitizedString.Replace(TEXT("."), TEXT("-"));
+
+	return SanitizedString;
+}
+
+FTransform CalcModelSpaceTransform(const FTransform* OutBoneTransform, int32 BoneIndex)
+{
+	FTransform BoneTransformMS = OutBoneTransform[BoneIndex];
+
+	while (BoneIndex != -1)
+	{
+		int32 parentIndex = SteamVRSkeleton::GetParentIndex(BoneIndex);
+		if (parentIndex != -1)
+		{
+			BoneTransformMS = BoneTransformMS * OutBoneTransform[parentIndex];
+			BoneIndex = parentIndex;
+		}
+		else
+		{
+			break;
+		}
+	}
+
+	return BoneTransformMS;
+}
+
+bool FSteamVRInputDevice::GetSkeletalData(bool bLeftHand, bool bMirror, EVRSkeletalMotionRange MotionRange, FTransform* OutBoneTransform, int32 OutBoneTransformCount)
+{
+	// Check that the size of the buffer we will be writing into is big enough to hold all the bone transforms
+	if (OutBoneTransformCount < STEAMVR_SKELETON_BONE_COUNT)
+	{
+		return false;
+	}
+
+	if (VRSystem() && VRInput())
+	{
+		// Get the handle for the skeletal action.  If its invalid (the necessary skeletal action is not in the manifest) then return false
+		vr::VRActionHandle_t ActionHandle = (bLeftHand) ? VRSkeletalHandleLeft : VRSkeletalHandleRight;
+		if (ActionHandle == k_ulInvalidActionHandle)
+		{
+			return false;
+		}
+
+		// Get skeletal data
+		VRBoneTransform_t SteamVRBoneTransforms[STEAMVR_SKELETON_BONE_COUNT];
+		EVRInputError Err = VRInput()->GetSkeletalBoneData(ActionHandle, vr::EVRSkeletalTransformSpace::VRSkeletalTransformSpace_Parent, MotionRange, SteamVRBoneTransforms, STEAMVR_SKELETON_BONE_COUNT);
+		
+		if (Err != VRInputError_None)
+		{
+			return false;
+		}
+
+		// Optionally mirror the pose to the opposite hand
+		if (bMirror)
+		{
+			MirrorSteamVRSkeleton(SteamVRBoneTransforms, STEAMVR_SKELETON_BONE_COUNT);
+		}
+
+		// GetSkeletalBoneData returns bone transforms are in SteamVR's coordinate system, so
+		// we need to convert them to UE4's coordinate system.  
+		// SteamVR coords:	X=right,	Y=up,		Z=backwards,	right-handed,	scale is meters
+		// UE4 coords:		X=forward,	Y=right,	Z=up,			left-handed,	scale is centimeters
+
+		// The root is positioned at the controller's anchor position with zero rotation.  
+		// However because of the conversion from SteamVR coordinates to Unreal coordinates the root bone is scaled
+		// to the new coordinate system
+		FTransform& RootTransform = OutBoneTransform[ESteamVRBone_Root];
+		RootTransform.SetComponents(FQuat::Identity, FVector::ZeroVector, FVector(100.f, 100.f, 100.f));
+
+		// Transform all the non-root bones to the new coordinate system
+		for (int32 BoneIndex = ESteamVRBone_Root + 1; BoneIndex < STEAMVR_SKELETON_BONE_COUNT; ++BoneIndex)
+		{
+			const VRBoneTransform_t& SrcTransform = SteamVRBoneTransforms[BoneIndex];
+
+			FQuat NewRotation(
+				SrcTransform.orientation.z,
+				-SrcTransform.orientation.x,
+				SrcTransform.orientation.y,
+				-SrcTransform.orientation.w
+			);
+
+			FVector NewTranslation(
+				SrcTransform.position.v[2],
+				-SrcTransform.position.v[0],
+				SrcTransform.position.v[1]
+			);
+
+			FTransform& DstTransform = OutBoneTransform[BoneIndex];
+			DstTransform.SetRotation(NewRotation);
+			DstTransform.SetTranslation(NewTranslation);
+		}
+
+		// Apply an extra transformation to the children of the root bone to compensate for the changes made to the root
+		// to make it fit the new coordinate system even though it has zero rotation
+		FQuat FixupRotation(FVector(0.f, 0.f, 1.f), PI);
+
+		for (int32 ChildIndex = 0; ChildIndex < SteamVRSkeleton::GetChildCount(ESteamVRBone_Root); ++ChildIndex)
+		{
+			int32 BoneIndex = SteamVRSkeleton::GetChildIndex(ESteamVRBone_Root, ChildIndex);
+
+			FTransform& DstTransform = OutBoneTransform[BoneIndex];
+
+			FVector NewTranslation = DstTransform.GetTranslation() * FVector(-1.f, -1.f, 1.f);
+			FQuat NewRotation = FixupRotation * DstTransform.GetRotation();
+
+			DstTransform.SetRotation(NewRotation);
+			DstTransform.SetTranslation(NewTranslation);
+		}
+
+		return true;
+	}
+
+	return false;
+}
+
+void FSteamVRInputDevice::SendAnalogMessage(const ETrackedControllerRole TrackedControllerRole, const FGamepadKeyNames::Type AxisButton, float AnalogValue)
+{
+	if (TrackedControllerRole == ETrackedControllerRole::TrackedControllerRole_LeftHand && bCurlsAndSplaysEnabled_L)
+	{
+		MessageHandler->OnControllerAnalog(AxisButton, 0, AnalogValue);
+		//UE_LOG(LogSteamVRInputDevice, Warning, TEXT("Left Index value: %f for axis %s"), ControllerState.IndexGripAnalog, *AxisButton.ToString());
+	}
+	else if (TrackedControllerRole == ETrackedControllerRole::TrackedControllerRole_RightHand && bCurlsAndSplaysEnabled_R)
+	{
+		MessageHandler->OnControllerAnalog(AxisButton, 0, AnalogValue);
+		//UE_LOG(LogSteamVRInputDevice, Warning, TEXT("Left Index value: %f for axis %s"), ControllerState.IndexGripAnalog, *AxisButton.ToString());
+	}
+}
+
+void FSteamVRInputDevice::SendControllerEvents()
+{
+
+	if (VRSystem() && VRInput() && SteamVRInputActionSets.Num() > 0)
+	{
+		EVRInputError ActionStateError = VRInput()->UpdateActionState(ActiveActionSets, sizeof(VRActiveActionSet_t), 1);
+
+		if (ActionStateError != VRInputError_None)
+		{
+			//GetInputError(ActionStateError, TEXT("Error encountered when trying to update the action state"));
+			return;
+		}
+
+		// Go through all Actions in all active ActionSets
+		for (auto& SteamVRInputActionSet : SteamVRInputActionSets)
+		{
+			ProcessActionEvents(SteamVRInputActionSet);
+		}
+	}
+}
+
+void FSteamVRInputDevice::SetMessageHandler(const TSharedRef<FGenericApplicationMessageHandler>& InMessageHandler)
+{
+	MessageHandler = InMessageHandler;
+}
+
+bool FSteamVRInputDevice::Exec(UWorld* InWorld, const TCHAR* Cmd, FOutputDevice& Ar)
+{
+	return false;
+}
+
+bool FSteamVRInputDevice::GetControllerOrientationAndPosition(const int32 ControllerIndex, const EControllerHand DeviceHand, FRotator& OutOrientation, FVector& OutPosition, float WorldToMetersScale) const
+{
+	if (VRInput() && VRCompositor())
+	{
+		InputPoseActionData_t PoseData = {};
+		EVRInputError InputError = VRInputError_NoData;
+
+		VRActionHandle_t LeftActionHandle = bUseSkeletonPose ? VRSkeletalHandleLeft : VRControllerHandleLeft;
+		VRActionHandle_t RightActionHandle = bUseSkeletonPose ? VRSkeletalHandleRight : VRControllerHandleRight;
+
+		switch (DeviceHand)
+		{
+		case EControllerHand::Left:
+			if (LeftActionHandle != vr::k_ulInvalidActionHandle)
+			{
+				if (GlobalPredictedSecondsFromNow <= -9999.f)
+				{
+					InputError = VRInput()->GetPoseActionDataForNextFrame(LeftActionHandle, VRCompositor()->GetTrackingSpace(), &PoseData, sizeof(PoseData), k_ulInvalidInputValueHandle);
+				}
+				else
+				{
+					InputError = VRInput()->GetPoseActionDataRelativeToNow(LeftActionHandle, VRCompositor()->GetTrackingSpace(), GlobalPredictedSecondsFromNow, &PoseData, sizeof(PoseData), k_ulInvalidInputValueHandle);
+				}
+
+				if (InputError != VRInputError_None)
+				{
+					return false;
+				}
+			}
+			break;
+		case EControllerHand::Right:
+
+			if (RightActionHandle != vr::k_ulInvalidActionHandle)
+			{
+				if (GlobalPredictedSecondsFromNow <= -9999.f)
+				{
+					InputError = VRInput()->GetPoseActionDataForNextFrame(RightActionHandle, VRCompositor()->GetTrackingSpace(), &PoseData, sizeof(PoseData), k_ulInvalidInputValueHandle);
+				}
+				else
+				{
+					InputError = VRInput()->GetPoseActionDataRelativeToNow(RightActionHandle, VRCompositor()->GetTrackingSpace(), GlobalPredictedSecondsFromNow, &PoseData, sizeof(PoseData), k_ulInvalidInputValueHandle);
+				}
+
+				if (InputError != VRInputError_None)
+				{
+					return false;
+				}
+			}
+			break;
+		case EControllerHand::Special_1:
+
+			if (VRSpecial1 == k_ulInvalidActionHandle)
+			{
+				return false;
+			}
+
+			if (GlobalPredictedSecondsFromNow <= -9999.f)
+			{
+				InputError = VRInput()->GetPoseActionDataForNextFrame(VRSpecial1, VRCompositor()->GetTrackingSpace(), &PoseData, sizeof(PoseData), k_ulInvalidInputValueHandle);
+			}
+			else
+			{
+				InputError = VRInput()->GetPoseActionDataRelativeToNow(VRSpecial1, VRCompositor()->GetTrackingSpace(), GlobalPredictedSecondsFromNow, &PoseData, sizeof(PoseData), k_ulInvalidInputValueHandle);
+			}
+
+			if (InputError != VRInputError_None)
+			{
+				return false;
+			}
+			break;
+		case EControllerHand::Special_2:
+
+			if (VRSpecial2 == k_ulInvalidActionHandle)
+			{
+				return false;
+			}
+
+			if (GlobalPredictedSecondsFromNow <= -9999.f)
+			{
+				InputError = VRInput()->GetPoseActionDataForNextFrame(VRSpecial2, VRCompositor()->GetTrackingSpace(), &PoseData, sizeof(PoseData), k_ulInvalidInputValueHandle);
+			}
+			else
+			{
+				InputError = VRInput()->GetPoseActionDataRelativeToNow(VRSpecial2, VRCompositor()->GetTrackingSpace(), GlobalPredictedSecondsFromNow, &PoseData, sizeof(PoseData), k_ulInvalidInputValueHandle);
+			}
+
+			if (InputError != VRInputError_None)
+			{
+				return false;
+			}
+			break;
+		case EControllerHand::Special_3:
+
+			if (VRSpecial3 == k_ulInvalidActionHandle)
+			{
+				return false;
+			}
+
+			if (GlobalPredictedSecondsFromNow <= -9999.f)
+			{
+				InputError = VRInput()->GetPoseActionDataForNextFrame(VRSpecial3, VRCompositor()->GetTrackingSpace(), &PoseData, sizeof(PoseData), k_ulInvalidInputValueHandle);
+			}
+			else
+			{
+				InputError = VRInput()->GetPoseActionDataRelativeToNow(VRSpecial3, VRCompositor()->GetTrackingSpace(), GlobalPredictedSecondsFromNow, &PoseData, sizeof(PoseData), k_ulInvalidInputValueHandle);
+			}
+
+			if (InputError != VRInputError_None)
+			{
+				return false;
+			}
+
+			break;
+		case EControllerHand::Special_4:
+
+			if (VRSpecial4 == k_ulInvalidActionHandle)
+			{
+				return false;
+			}
+
+			if (GlobalPredictedSecondsFromNow <= -9999.f)
+			{
+				InputError = VRInput()->GetPoseActionDataForNextFrame(VRSpecial4, VRCompositor()->GetTrackingSpace(), &PoseData, sizeof(PoseData), k_ulInvalidInputValueHandle);
+			}
+			else
+			{
+				InputError = VRInput()->GetPoseActionDataRelativeToNow(VRSpecial4, VRCompositor()->GetTrackingSpace(), GlobalPredictedSecondsFromNow, &PoseData, sizeof(PoseData), k_ulInvalidInputValueHandle);
+			}
+
+			if (InputError != VRInputError_None)
+			{
+				return false;
+			}
+
+			break;
+		case EControllerHand::Special_5:
+
+			if (VRSpecial5 == k_ulInvalidActionHandle)
+			{
+				return false;
+			}
+
+			if (GlobalPredictedSecondsFromNow <= -9999.f)
+			{
+				InputError = VRInput()->GetPoseActionDataForNextFrame(VRSpecial5, VRCompositor()->GetTrackingSpace(), &PoseData, sizeof(PoseData), k_ulInvalidInputValueHandle);
+			}
+			else
+			{
+				InputError = VRInput()->GetPoseActionDataRelativeToNow(VRSpecial5, VRCompositor()->GetTrackingSpace(), GlobalPredictedSecondsFromNow, &PoseData, sizeof(PoseData), k_ulInvalidInputValueHandle);
+			}
+
+			if (InputError != VRInputError_None)
+			{
+				return false;
+			}
+
+			break;
+		case EControllerHand::Special_6:
+
+			if (VRSpecial6 == k_ulInvalidActionHandle)
+			{
+				return false;
+			}
+
+			if (GlobalPredictedSecondsFromNow <= -9999.f)
+			{
+				InputError = VRInput()->GetPoseActionDataForNextFrame(VRSpecial6, VRCompositor()->GetTrackingSpace(), &PoseData, sizeof(PoseData), k_ulInvalidInputValueHandle);
+			}
+			else
+			{
+				InputError = VRInput()->GetPoseActionDataRelativeToNow(VRSpecial6, VRCompositor()->GetTrackingSpace(), GlobalPredictedSecondsFromNow, &PoseData, sizeof(PoseData), k_ulInvalidInputValueHandle);
+			}
+
+			if (InputError != VRInputError_None)
+			{
+				return false;
+			}
+
+			break;
+		case EControllerHand::Special_7:
+
+			if (VRSpecial7 == k_ulInvalidActionHandle)
+			{
+				return false;
+			}
+
+			if (GlobalPredictedSecondsFromNow <= -9999.f)
+			{
+				InputError = VRInput()->GetPoseActionDataForNextFrame(VRSpecial7, VRCompositor()->GetTrackingSpace(), &PoseData, sizeof(PoseData), k_ulInvalidInputValueHandle);
+			}
+			else
+			{
+				InputError = VRInput()->GetPoseActionDataRelativeToNow(VRSpecial7, VRCompositor()->GetTrackingSpace(), GlobalPredictedSecondsFromNow, &PoseData, sizeof(PoseData), k_ulInvalidInputValueHandle);
+			}
+
+			if (InputError != VRInputError_None)
+			{
+				return false;
+			}
+
+			break;
+		case EControllerHand::Special_8:
+
+			if (VRSpecial8 == k_ulInvalidActionHandle)
+			{
+				return false;
+			}
+
+			if (GlobalPredictedSecondsFromNow <= -9999.f)
+			{
+				InputError = VRInput()->GetPoseActionDataForNextFrame(VRSpecial8, VRCompositor()->GetTrackingSpace(), &PoseData, sizeof(PoseData), k_ulInvalidInputValueHandle);
+			}
+			else
+			{
+				InputError = VRInput()->GetPoseActionDataRelativeToNow(VRSpecial8, VRCompositor()->GetTrackingSpace(), GlobalPredictedSecondsFromNow, &PoseData, sizeof(PoseData), k_ulInvalidInputValueHandle);
+			}
+
+			if (InputError != VRInputError_None)
+			{
+				return false;
+			}
+
+			break;
+		default:
+			return false;
+		}
+
+		if (InputError == VRInputError_None)
+		{
+			// Get SteamVR Transform Matrix for this controller
+			HmdMatrix34_t Matrix = PoseData.pose.mDeviceToAbsoluteTracking;
+
+			// Transform SteamVR Pose to Unreal Pose
+			FMatrix Pose = FMatrix(
+				FPlane(Matrix.m[0][0], Matrix.m[1][0], Matrix.m[2][0], 0.0f),
+				FPlane(Matrix.m[0][1], Matrix.m[1][1], Matrix.m[2][1], 0.0f),
+				FPlane(Matrix.m[0][2], Matrix.m[1][2], Matrix.m[2][2], 0.0f),
+				FPlane(Matrix.m[0][3], Matrix.m[1][3], Matrix.m[2][3], 1.0f)
+			);
+
+
+			// Transform SteamVR Rotation Quaternion to a UE FRotator
+			FQuat OrientationQuat;
+			FQuat Orientation(Pose);
+			OrientationQuat.X = -Orientation.Z;
+			OrientationQuat.Y = Orientation.X;
+			OrientationQuat.Z = Orientation.Y;
+			OrientationQuat.W = -Orientation.W;
+
+			// Return controller transform
+			FVector Position = ((FVector(-Pose.M[3][2], Pose.M[3][0], Pose.M[3][1])) * WorldToMetersScale - CachedBasePosition);
+			OutPosition = CachedBaseOrientation.Inverse().RotateVector(Position);
+
+			OrientationQuat = CachedBaseOrientation.Inverse() * OrientationQuat;
+			OrientationQuat.Normalize();
+			OutOrientation = OrientationQuat.Rotator();
+
+		}
+	}
+
+	return true;
+}
+
+ETrackingStatus FSteamVRInputDevice::GetControllerTrackingStatus(const int32 ControllerIndex, const EControllerHand DeviceHand) const
+{
+	ETrackingStatus TrackingStatus = ETrackingStatus::NotTracked;
+
+	if (VRInput() && VRCompositor())
+	{
+		InputPoseActionData_t PoseData = {};
+		EVRInputError InputError = VRInputError_NoData;
+
+		switch (DeviceHand)
+		{
+		case EControllerHand::Left:
+
+			if (VRControllerHandleLeft == k_ulInvalidActionHandle)
+			{
+				return ETrackingStatus::NotTracked;
+			}
+
+			if (GlobalPredictedSecondsFromNow <= -9999.f)
+			{
+				InputError = VRInput()->GetPoseActionDataForNextFrame(VRControllerHandleLeft, VRCompositor()->GetTrackingSpace(), &PoseData, sizeof(PoseData), k_ulInvalidInputValueHandle);
+			}
+			else
+			{
+				InputError = VRInput()->GetPoseActionDataRelativeToNow(VRControllerHandleLeft, VRCompositor()->GetTrackingSpace(), GlobalPredictedSecondsFromNow, &PoseData, sizeof(PoseData), k_ulInvalidInputValueHandle);
+			}
+
+			if (InputError != VRInputError_None)
+			{
+				return ETrackingStatus::NotTracked;
+			}
+
+			break;
+		case EControllerHand::Right:
+
+			if (VRControllerHandleRight == k_ulInvalidActionHandle)
+			{
+				return ETrackingStatus::NotTracked;
+			}
+
+			if (GlobalPredictedSecondsFromNow <= -9999.f)
+			{
+				InputError = VRInput()->GetPoseActionDataForNextFrame(VRControllerHandleRight, VRCompositor()->GetTrackingSpace(), &PoseData, sizeof(PoseData), k_ulInvalidInputValueHandle);
+			}
+			else
+			{
+				InputError = VRInput()->GetPoseActionDataRelativeToNow(VRControllerHandleRight, VRCompositor()->GetTrackingSpace(), GlobalPredictedSecondsFromNow, &PoseData, sizeof(PoseData), k_ulInvalidInputValueHandle);
+			}
+
+			if (InputError != VRInputError_None)
+			{
+				return ETrackingStatus::NotTracked;
+			}
+
+			break;
+		case EControllerHand::Special_1:
+
+			if (VRSpecial1 == k_ulInvalidActionHandle)
+			{
+				return ETrackingStatus::NotTracked;
+			}
+
+			if (GlobalPredictedSecondsFromNow <= -9999.f)
+			{
+				InputError = VRInput()->GetPoseActionDataForNextFrame(VRSpecial1, VRCompositor()->GetTrackingSpace(), &PoseData, sizeof(PoseData), k_ulInvalidInputValueHandle);
+			}
+			else
+			{
+				InputError = VRInput()->GetPoseActionDataRelativeToNow(VRSpecial1, VRCompositor()->GetTrackingSpace(), GlobalPredictedSecondsFromNow, &PoseData, sizeof(PoseData), k_ulInvalidInputValueHandle);
+			}
+
+			if (InputError != VRInputError_None)
+			{
+				return ETrackingStatus::NotTracked;
+			}
+
+			break;
+		case EControllerHand::Special_2:
+
+			if (VRSpecial2 == k_ulInvalidActionHandle)
+			{
+				return ETrackingStatus::NotTracked;
+			}
+
+			if (GlobalPredictedSecondsFromNow <= -9999.f)
+			{
+				InputError = VRInput()->GetPoseActionDataForNextFrame(VRSpecial2, VRCompositor()->GetTrackingSpace(), &PoseData, sizeof(PoseData), k_ulInvalidInputValueHandle);
+			}
+			else
+			{
+				InputError = VRInput()->GetPoseActionDataRelativeToNow(VRSpecial2, VRCompositor()->GetTrackingSpace(), GlobalPredictedSecondsFromNow, &PoseData, sizeof(PoseData), k_ulInvalidInputValueHandle);
+			}
+
+			if (InputError != VRInputError_None)
+			{
+				return ETrackingStatus::NotTracked;
+			}
+
+			break;
+		case EControllerHand::Special_3:
+
+			if (VRSpecial3 == k_ulInvalidActionHandle)
+			{
+				return ETrackingStatus::NotTracked;
+			}
+
+			if (GlobalPredictedSecondsFromNow <= -9999.f)
+			{
+				InputError = VRInput()->GetPoseActionDataForNextFrame(VRSpecial3, VRCompositor()->GetTrackingSpace(), &PoseData, sizeof(PoseData), k_ulInvalidInputValueHandle);
+			}
+			else
+			{
+				InputError = VRInput()->GetPoseActionDataRelativeToNow(VRSpecial3, VRCompositor()->GetTrackingSpace(), GlobalPredictedSecondsFromNow, &PoseData, sizeof(PoseData), k_ulInvalidInputValueHandle);
+			}
+
+			if (InputError != VRInputError_None)
+			{
+				return ETrackingStatus::NotTracked;
+			}
+
+			break;
+		case EControllerHand::Special_4:
+
+			if (VRSpecial4 == k_ulInvalidActionHandle)
+			{
+				return ETrackingStatus::NotTracked;
+			}
+
+			if (GlobalPredictedSecondsFromNow <= -9999.f)
+			{
+				InputError = VRInput()->GetPoseActionDataForNextFrame(VRSpecial4, VRCompositor()->GetTrackingSpace(), &PoseData, sizeof(PoseData), k_ulInvalidInputValueHandle);
+			}
+			else
+			{
+				InputError = VRInput()->GetPoseActionDataRelativeToNow(VRSpecial4, VRCompositor()->GetTrackingSpace(), GlobalPredictedSecondsFromNow, &PoseData, sizeof(PoseData), k_ulInvalidInputValueHandle);
+			}
+
+			if (InputError != VRInputError_None)
+			{
+				return ETrackingStatus::NotTracked;
+			}
+
+			break;
+		case EControllerHand::Special_5:
+
+			if (VRSpecial5 == k_ulInvalidActionHandle)
+			{
+				return ETrackingStatus::NotTracked;
+			}
+
+			if (GlobalPredictedSecondsFromNow <= -9999.f)
+			{
+				InputError = VRInput()->GetPoseActionDataForNextFrame(VRSpecial5, VRCompositor()->GetTrackingSpace(), &PoseData, sizeof(PoseData), k_ulInvalidInputValueHandle);
+			}
+			else
+			{
+				InputError = VRInput()->GetPoseActionDataRelativeToNow(VRSpecial5, VRCompositor()->GetTrackingSpace(), GlobalPredictedSecondsFromNow, &PoseData, sizeof(PoseData), k_ulInvalidInputValueHandle);
+			}
+
+			if (InputError != VRInputError_None)
+			{
+				return ETrackingStatus::NotTracked;
+			}
+
+			break;
+		case EControllerHand::Special_6:
+
+			if (VRSpecial6 == k_ulInvalidActionHandle)
+			{
+				return ETrackingStatus::NotTracked;
+			}
+
+			if (GlobalPredictedSecondsFromNow <= -9999.f)
+			{
+				InputError = VRInput()->GetPoseActionDataForNextFrame(VRSpecial6, VRCompositor()->GetTrackingSpace(), &PoseData, sizeof(PoseData), k_ulInvalidInputValueHandle);
+			}
+			else
+			{
+				InputError = VRInput()->GetPoseActionDataRelativeToNow(VRSpecial6, VRCompositor()->GetTrackingSpace(), GlobalPredictedSecondsFromNow, &PoseData, sizeof(PoseData), k_ulInvalidInputValueHandle);
+			}
+
+			if (InputError != VRInputError_None)
+			{
+				return ETrackingStatus::NotTracked;
+			}
+
+			break;
+		case EControllerHand::Special_7:
+
+			if (VRSpecial7 == k_ulInvalidActionHandle)
+			{
+				return ETrackingStatus::NotTracked;
+			}
+
+			if (GlobalPredictedSecondsFromNow <= -9999.f)
+			{
+				InputError = VRInput()->GetPoseActionDataForNextFrame(VRSpecial7, VRCompositor()->GetTrackingSpace(), &PoseData, sizeof(PoseData), k_ulInvalidInputValueHandle);
+			}
+			else
+			{
+				InputError = VRInput()->GetPoseActionDataRelativeToNow(VRSpecial7, VRCompositor()->GetTrackingSpace(), GlobalPredictedSecondsFromNow, &PoseData, sizeof(PoseData), k_ulInvalidInputValueHandle);
+			}
+
+			if (InputError != VRInputError_None)
+			{
+				return ETrackingStatus::NotTracked;
+			}
+
+			break;
+		case EControllerHand::Special_8:
+
+			if (VRSpecial8 == k_ulInvalidActionHandle)
+			{
+				return ETrackingStatus::NotTracked;
+			}
+
+			if (GlobalPredictedSecondsFromNow <= -9999.f)
+			{
+				InputError = VRInput()->GetPoseActionDataForNextFrame(VRSpecial8, VRCompositor()->GetTrackingSpace(), &PoseData, sizeof(PoseData), k_ulInvalidInputValueHandle);
+			}
+			else
+			{
+				InputError = VRInput()->GetPoseActionDataRelativeToNow(VRSpecial8, VRCompositor()->GetTrackingSpace(), GlobalPredictedSecondsFromNow, &PoseData, sizeof(PoseData), k_ulInvalidInputValueHandle);
+			}
+
+			if (InputError != VRInputError_None)
+			{
+				return ETrackingStatus::NotTracked;
+			}
+
+			break;
+		default:
+			break;
+		}
+
+		if (InputError == VRInputError_None && PoseData.pose.bDeviceIsConnected)
+		{
+			TrackingStatus = ETrackingStatus::Tracked;
+		}
+	}
+
+	return TrackingStatus;
+}
+
+FName FSteamVRInputDevice::GetMotionControllerDeviceTypeName() const
+{
+	return FName(TEXT("SteamVRInputDevice"));
+}
+
+bool FSteamVRInputDevice::GetHandJointPosition(const FName MotionSource, int jointIndex, FVector& OutPosition) const
+{
+	return false;
+}
+
+void FSteamVRInputDevice::SetHapticFeedbackValues(int32 ControllerId, int32 Hand, const FHapticFeedbackValues& Values)
+{
+	VRActionHandle_t VibrationAction = k_ulInvalidActionHandle;
+
+	switch (Hand)
+	{
+		case (int32)EControllerHand::Left:
+			VibrationAction = VRVibrationLeft;
+			break;
+		case (int32)EControllerHand::Right:
+			VibrationAction = VRVibrationRight;
+			break;
+		case (int32)EControllerHand::AnyHand:
+			VibrationAction = VRVibrationLeft;	// UE4.17+: Hardwire AnyHand to OpenVR's left path as it is the lowest device id to cover most use cases without triggering a duplicate vibration/rumble. TODO: May need refactor for cases where there's a left hand device set & right hand was set to AnyHand for some reason.
+			break;
+	}
+
+	if (VRSystem() && VRInput() && VibrationAction != k_ulInvalidActionHandle)
+	{
+		VRInput()->TriggerHapticVibrationAction(VibrationAction, 0.f, CurrentDeltaTime, Values.Frequency, Values.Amplitude, k_ulInvalidInputValueHandle);
+		//UE_LOG(LogSteamVRInputDevice, Warning, TEXT("[HAPTIC] Hand: %i, Duration: %f, Frequency: %f, Amplitude: %f"), Hand, CurrentDeltaTime, Values.Frequency, Values.Amplitude);
+	}
+}
+
+void FSteamVRInputDevice::GetHapticFrequencyRange(float& MinFrequency, float& MaxFrequency) const
+{
+	MinFrequency = MaxFrequency = 0.f;
+}
+
+float FSteamVRInputDevice::GetHapticAmplitudeScale() const
+{
+	return 1.f;
+}
+
+void FSteamVRInputDevice::GetControllerFidelity()
+{
+	if (VRInput() && VRCompositor())
+	{
+		InputPoseActionData_t PoseData = {};
+		EVRInputError InputError = VRInputError_NoData;
+
+		if (VRControllerHandleLeft == k_ulInvalidActionHandle)
+		{
+			return;
+		}
+
+		InputError = VRInput()->GetPoseActionDataForNextFrame(VRControllerHandleLeft, VRCompositor()->GetTrackingSpace(), &PoseData, sizeof(PoseData), k_ulInvalidInputValueHandle);
+		
+		if (InputError != VRInputError_None)
+		{
+			return;
+		}
+		
+		if (PoseData.bActive && PoseData.pose.bDeviceIsConnected)
+		{
+			if (VRSkeletalHandleLeft == k_ulInvalidActionHandle)
+			{
+				return;
+			}
+
+			VRInput()->GetSkeletalTrackingLevel(VRSkeletalHandleLeft, &LeftControllerFidelity);
+
+			if (InputError != VRInputError_None)
+			{
+				return;
+			}
+
+			bIsSkeletalControllerLeftPresent = (LeftControllerFidelity >= VRSkeletalTracking_Partial);
+		}
+		else
+		{
+			bIsSkeletalControllerLeftPresent = false;
+			LeftControllerFidelity = EVRSkeletalTrackingLevel::VRSkeletalTracking_Estimated;
+		}
+
+		if (VRControllerHandleRight == k_ulInvalidActionHandle)
+		{
+			return;
+		}
+
+		InputError = VRInput()->GetPoseActionDataForNextFrame(VRControllerHandleRight, VRCompositor()->GetTrackingSpace(), &PoseData, sizeof(PoseData), k_ulInvalidInputValueHandle);
+		if (PoseData.bActive && PoseData.pose.bDeviceIsConnected)
+		{
+			if (VRSkeletalHandleRight == k_ulInvalidActionHandle)
+			{
+				return;
+			}
+
+			VRInput()->GetSkeletalTrackingLevel(VRSkeletalHandleRight, &RightControllerFidelity);
+
+			if (InputError != VRInputError_None)
+			{
+				return;
+			}
+
+			bIsSkeletalControllerRightPresent = (RightControllerFidelity >= VRSkeletalTracking_Partial);
+		} 
+		else
+		{
+			bIsSkeletalControllerRightPresent = false;
+			RightControllerFidelity = EVRSkeletalTrackingLevel::VRSkeletalTracking_Estimated;
+		}
+	}
+}
+
+void FSteamVRInputDevice::GetLeftHandPoseData(FVector& Position, FRotator& Orientation, FVector& AngularVelocity, FVector& Velocity)
+{
+	InputPoseActionData_t PoseData = {};
+	EVRInputError InputError = VRInputError_NoData;
+
+	if (bIsSkeletalControllerRightPresent && VRInput())
+	{
+		if (VRSkeletalHandleLeft == k_ulInvalidActionHandle)
+		{
+			return;
+		}
+
+		InputError = VRInput()->GetPoseActionDataForNextFrame(VRSkeletalHandleLeft, VRCompositor()->GetTrackingSpace(), &PoseData, sizeof(PoseData), k_ulInvalidInputValueHandle);
+		
+		if (InputError != VRInputError_None || VRSkeletalHandleLeft == k_ulInvalidActionHandle)
+		{
+			return;
+		}
+		
+		if (PoseData.bActive && PoseData.pose.bDeviceIsConnected && InputError == VRInputError_None)
+		{
+			GetUETransform(PoseData, Position, Orientation);
+			AngularVelocity = FVector(
+				PoseData.pose.vAngularVelocity.v[2],
+				-PoseData.pose.vAngularVelocity.v[0],
+				PoseData.pose.vAngularVelocity.v[1]
+			);
+			Velocity = FVector(
+				PoseData.pose.vVelocity.v[2],
+				-PoseData.pose.vVelocity.v[0],
+				PoseData.pose.vVelocity.v[1]
+			);
+		}
+	}
+}
+
+void FSteamVRInputDevice::GetRightHandPoseData(FVector& Position, FRotator& Orientation, FVector& AngularVelocity, FVector& Velocity)
+{
+	InputPoseActionData_t PoseData = {};
+	EVRInputError InputError = VRInputError_NoData;
+
+	if (bIsSkeletalControllerRightPresent && VRInput())
+	{
+		if (VRSkeletalHandleRight == k_ulInvalidActionHandle)
+		{
+			return;
+		}
+
+		InputError = VRInput()->GetPoseActionDataForNextFrame(VRSkeletalHandleRight, VRCompositor()->GetTrackingSpace(), &PoseData, sizeof(PoseData), k_ulInvalidInputValueHandle);
+		
+		if (InputError != VRInputError_None)
+		{
+			return;
+		}
+
+		if (PoseData.bActive && PoseData.pose.bDeviceIsConnected && InputError == VRInputError_None)
+		{
+			GetUETransform(PoseData, Position, Orientation);
+			AngularVelocity = FVector(
+										PoseData.pose.vAngularVelocity.v[2],
+										-PoseData.pose.vAngularVelocity.v[0],
+										PoseData.pose.vAngularVelocity.v[1]
+									 );
+			Velocity = FVector(
+								PoseData.pose.vVelocity.v[2],
+								-PoseData.pose.vVelocity.v[0],
+								PoseData.pose.vVelocity.v[1]
+								);
+		}
+	}
+}
+
+void FSteamVRInputDevice::GetUETransform(InputPoseActionData_t PoseData, FVector& OutPosition, FRotator& OutOrientation)
+{
+	// Get SteamVR Transform Matrix for this skeleton
+	HmdMatrix34_t Matrix = PoseData.pose.mDeviceToAbsoluteTracking;
+
+	// Transform SteamVR Pose to Unreal Pose
+	FMatrix Pose = FMatrix(
+		FPlane(Matrix.m[0][0], Matrix.m[1][0], Matrix.m[2][0], 0.0f),
+		FPlane(Matrix.m[0][1], Matrix.m[1][1], Matrix.m[2][1], 0.0f),
+		FPlane(Matrix.m[0][2], Matrix.m[1][2], Matrix.m[2][2], 0.0f),
+		FPlane(Matrix.m[0][3], Matrix.m[1][3], Matrix.m[2][3], 1.0f)
+	);
+
+
+	// Transform SteamVR Rotation Quaternion to a UE FRotator
+	FQuat OrientationQuat;
+	FQuat Orientation(Pose);
+	OrientationQuat.X = -Orientation.Z;
+	OrientationQuat.Y = Orientation.X;
+	OrientationQuat.Z = Orientation.Y;
+	OrientationQuat.W = -Orientation.W;
+
+
+	FVector Position = ((FVector(-Pose.M[3][2], Pose.M[3][0], Pose.M[3][1])) * GWorld->GetWorldSettings()->WorldToMeters);
+	OutPosition = Position;
+
+	//OutOrientation = BaseOrientation.Inverse() * OutOrientation;
+	OutOrientation.Normalize();
+	OutOrientation = OrientationQuat.Rotator();			
+}
+
+void FSteamVRInputDevice::SetChannelValue(int32 ControllerId, FForceFeedbackChannelType ChannelType, float Value)
+{
+	// Empty on purpose
+}
+
+void FSteamVRInputDevice::SetChannelValues(int32 ControllerId, const FForceFeedbackValues &values)
+{
+	// Empty on purpose
+}
+
+void FSteamVRInputDevice::InitControllerMappings()
+{
+	for (unsigned int i = 0; i < k_unMaxTrackedDeviceCount; ++i)
+	{
+		DeviceToControllerMap[i] = INDEX_NONE;
+	}
+
+	for (unsigned int id = 0; id < SteamVRInputDeviceConstants::MaxUnrealControllers; ++id)
+	{
+		for (unsigned int hand = 0; hand < k_unMaxTrackedDeviceCount; ++hand)
+		{
+			if (DeviceToControllerMap[id] < SteamVRInputDeviceConstants::MaxUnrealControllers && hand < k_unMaxTrackedDeviceCount)
+			{
+				UnrealControllerIdAndHandToDeviceIdMap[DeviceToControllerMap[id]][hand] = INDEX_NONE;
+			}
+		}
+	}
+
+	for (int32& HandCount : MaxUEHandCount)
+	{
+		HandCount = 0;
+	}
+}
+
+void FSteamVRInputDevice::InitControllerKeys()
+{
+	EKeys::AddMenuCategoryDisplayInfo("SteamVRInput", LOCTEXT("SteamVRInputSubCategory", "SteamVR Input"), TEXT("GraphEditor.PadEvent_16x"));
+
+#pragma region GENERIC KEYS
+	EKeys::AddKey(FKeyDetails(GenericKeys::SteamVR_MotionController_None, LOCTEXT("SteamVR_MotionController_None", "SteamVR Generic Key"), FKeyDetails::GamepadKey | FKeyDetails::NotBlueprintBindableKey, "SteamVRInput"));
+	EKeys::AddKey(FKeyDetails(GenericKeys::SteamVR_HMD_Proximity, LOCTEXT("SteamVR_HMD_Proximity", "SteamVR HMD Proximity"), FKeyDetails::GamepadKey | FKeyDetails::NotBlueprintBindableKey, "SteamVRInput"));
+#pragma endregion
+
+#pragma region INDEX CONTROLLER
+	EKeys::AddKey(FKeyDetails(IndexControllerKeys::ValveIndex_Left_Thumbstick_Up, LOCTEXT("ValveIndex_Left_Thumbstick_Up", "Valve Index (L) Thumbstick Up"), FKeyDetails::GamepadKey | FKeyDetails::NotBlueprintBindableKey, "ValveIndex"));
+	EKeys::AddKey(FKeyDetails(IndexControllerKeys::ValveIndex_Left_Thumbstick_Down, LOCTEXT("ValveIndex_Left_Thumbstick_Down", "Valve Index (L) Thumbstick Down"), FKeyDetails::GamepadKey | FKeyDetails::NotBlueprintBindableKey, "ValveIndex"));
+	EKeys::AddKey(FKeyDetails(IndexControllerKeys::ValveIndex_Left_Thumbstick_Left, LOCTEXT("ValveIndex_Left_Thumbstick_Left", "Valve Index (L) Thumbstick Left"), FKeyDetails::GamepadKey | FKeyDetails::NotBlueprintBindableKey, "ValveIndex"));
+	EKeys::AddKey(FKeyDetails(IndexControllerKeys::ValveIndex_Left_Thumbstick_Right, LOCTEXT("ValveIndex_Left_Thumbstick_Right", "Valve Index (L) Thumbstick Right"), FKeyDetails::GamepadKey | FKeyDetails::NotBlueprintBindableKey, "ValveIndex"));
+
+	EKeys::AddKey(FKeyDetails(IndexControllerKeys::ValveIndex_Right_Thumbstick_Up, LOCTEXT("ValveIndex_Right_Thumbstick_Up", "Valve Index (R) Thumbstick Up"), FKeyDetails::GamepadKey | FKeyDetails::NotBlueprintBindableKey, "ValveIndex"));
+	EKeys::AddKey(FKeyDetails(IndexControllerKeys::ValveIndex_Right_Thumbstick_Down, LOCTEXT("ValveIndex_Right_Thumbstick_Down", "Valve Index (R) Thumbstick Down"), FKeyDetails::GamepadKey | FKeyDetails::NotBlueprintBindableKey, "ValveIndex"));
+	EKeys::AddKey(FKeyDetails(IndexControllerKeys::ValveIndex_Right_Thumbstick_Left, LOCTEXT("ValveIndex_Right_Thumbstick_Left", "Valve Index (R) Thumbstick Left"), FKeyDetails::GamepadKey | FKeyDetails::NotBlueprintBindableKey, "ValveIndex"));
+	EKeys::AddKey(FKeyDetails(IndexControllerKeys::ValveIndex_Right_Thumbstick_Right, LOCTEXT("ValveIndex_Right_Thumbstick_Right", "Valve Index (R) Thumbstick Right"), FKeyDetails::GamepadKey | FKeyDetails::NotBlueprintBindableKey, "ValveIndex"));
+
+	EKeys::AddKey(FKeyDetails(IndexControllerKeys::ValveIndex_Left_Trackpad_Up, LOCTEXT("ValveIndex_Left_Trackpad_Up", "Valve Index (L) Trackpad Up"), FKeyDetails::GamepadKey | FKeyDetails::NotBlueprintBindableKey, "ValveIndex"));
+	EKeys::AddKey(FKeyDetails(IndexControllerKeys::ValveIndex_Left_Trackpad_Down, LOCTEXT("ValveIndex_Left_Trackpad_Down", "Valve Index (L) Trackpad Down"), FKeyDetails::GamepadKey | FKeyDetails::NotBlueprintBindableKey, "ValveIndex"));
+	EKeys::AddKey(FKeyDetails(IndexControllerKeys::ValveIndex_Left_Trackpad_Left, LOCTEXT("ValveIndex_Left_Trackpad_Left", "Valve Index (L) Trackpad Left"), FKeyDetails::GamepadKey | FKeyDetails::NotBlueprintBindableKey, "ValveIndex"));
+	EKeys::AddKey(FKeyDetails(IndexControllerKeys::ValveIndex_Left_Trackpad_Right, LOCTEXT("ValveIndex_Left_Trackpad_Right", "Valve Index (L) Trackpad Right"), FKeyDetails::GamepadKey | FKeyDetails::NotBlueprintBindableKey, "ValveIndex"));
+
+	EKeys::AddKey(FKeyDetails(IndexControllerKeys::ValveIndex_Right_Trackpad_Up, LOCTEXT("ValveIndex_Right_Trackpad_Up", "Valve Index (R) Trackpad Up"), FKeyDetails::GamepadKey | FKeyDetails::NotBlueprintBindableKey, "ValveIndex"));
+	EKeys::AddKey(FKeyDetails(IndexControllerKeys::ValveIndex_Right_Trackpad_Down, LOCTEXT("ValveIndex_Right_Trackpad_Down", "Valve Index (R) Trackpad Down"), FKeyDetails::GamepadKey | FKeyDetails::NotBlueprintBindableKey, "ValveIndex"));
+	EKeys::AddKey(FKeyDetails(IndexControllerKeys::ValveIndex_Right_Trackpad_Left, LOCTEXT("ValveIndex_Right_Trackpad_Left", "Valve Index (R) Trackpad Left"), FKeyDetails::GamepadKey | FKeyDetails::NotBlueprintBindableKey, "ValveIndex"));
+	EKeys::AddKey(FKeyDetails(IndexControllerKeys::ValveIndex_Right_Trackpad_Right, LOCTEXT("ValveIndex_Right_Trackpad_Right", "Valve Index (R) Trackpad Right"), FKeyDetails::GamepadKey | FKeyDetails::NotBlueprintBindableKey, "ValveIndex"));
+
+	EKeys::AddKey(FKeyDetails(IndexControllerKeys::ValveIndex_Left_Pinch_Grab, LOCTEXT("ValveIndex_Left_Pinch_Grab", "Valve Index (L) Pinch Grab"), FKeyDetails::GamepadKey | FKeyDetails::NotBlueprintBindableKey, "ValveIndex"));
+	EKeys::AddKey(FKeyDetails(IndexControllerKeys::ValveIndex_Right_Pinch_Grab, LOCTEXT("ValveIndex_Right_Pinch_Grab", "Valve Index (R) Pinch Grab"), FKeyDetails::GamepadKey | FKeyDetails::NotBlueprintBindableKey, "ValveIndex"));
+
+	EKeys::AddKey(FKeyDetails(IndexControllerKeys::ValveIndex_Left_Grip_Grab, LOCTEXT("ValveIndex_Left_Grip_Grab", "Valve Index (L) Grip Grab"), FKeyDetails::GamepadKey | FKeyDetails::NotBlueprintBindableKey, "ValveIndex"));
+	EKeys::AddKey(FKeyDetails(IndexControllerKeys::ValveIndex_Right_Grip_Grab, LOCTEXT("ValveIndex_Right_Grip_Grab", "Valve Index (R) Grip Grab"), FKeyDetails::GamepadKey | FKeyDetails::NotBlueprintBindableKey, "ValveIndex"));
+#pragma endregion
+
+#pragma region OCULUS TOUCH CONTROLLER
+	EKeys::AddKey(FKeyDetails(OculusTouchKeys::OculusTouch_Left_Thumbstick_Up, LOCTEXT("OculusTouch_Left_Thumbstick_Up", "Oculus Touch (L) Thumbstick Up"), FKeyDetails::GamepadKey | FKeyDetails::NotBlueprintBindableKey, "OculusTouch"));
+	EKeys::AddKey(FKeyDetails(OculusTouchKeys::OculusTouch_Left_Thumbstick_Down, LOCTEXT("OculusTouch_Left_Thumbstick_Down", "Oculus Touch (L) Thumbstick Down"), FKeyDetails::GamepadKey | FKeyDetails::NotBlueprintBindableKey, "OculusTouch"));
+	EKeys::AddKey(FKeyDetails(OculusTouchKeys::OculusTouch_Left_Thumbstick_Left, LOCTEXT("OculusTouch_Left_Thumbstick_Left", "Oculus Touch (L) Thumbstick Left"), FKeyDetails::GamepadKey | FKeyDetails::NotBlueprintBindableKey, "OculusTouch"));
+	EKeys::AddKey(FKeyDetails(OculusTouchKeys::OculusTouch_Left_Thumbstick_Right, LOCTEXT("OculusTouch_Left_Thumbstick_Right", "Oculus Touch (L) Thumbstick Right"), FKeyDetails::GamepadKey | FKeyDetails::NotBlueprintBindableKey, "OculusTouch"));
+
+	EKeys::AddKey(FKeyDetails(OculusTouchKeys::OculusTouch_Right_Thumbstick_Up, LOCTEXT("OculusTouch_Right_Thumbstick_Up", "Oculus Touch (R) Thumbstick Up"), FKeyDetails::GamepadKey | FKeyDetails::NotBlueprintBindableKey, "OculusTouch"));
+	EKeys::AddKey(FKeyDetails(OculusTouchKeys::OculusTouch_Right_Thumbstick_Down, LOCTEXT("OculusTouch_Right_Thumbstick_Down", "Oculus Touch (R) Thumbstick Down"), FKeyDetails::GamepadKey | FKeyDetails::NotBlueprintBindableKey, "OculusTouch"));
+	EKeys::AddKey(FKeyDetails(OculusTouchKeys::OculusTouch_Right_Thumbstick_Left, LOCTEXT("OculusTouch_Right_Thumbstick_Left", "Oculus Touch (R) Thumbstick Left"), FKeyDetails::GamepadKey | FKeyDetails::NotBlueprintBindableKey, "OculusTouch"));
+	EKeys::AddKey(FKeyDetails(OculusTouchKeys::OculusTouch_Right_Thumbstick_Right, LOCTEXT("OculusTouch_Right_Thumbstick_Right", "Oculus Touch (R) Thumbstick Right"), FKeyDetails::GamepadKey | FKeyDetails::NotBlueprintBindableKey, "OculusTouch"));
+#pragma endregion
+
+#pragma region VIVE CONTROLLER
+	EKeys::AddKey(FKeyDetails(ViveControllerKeys::Vive_Left_Trackpad_Up, LOCTEXT("Vive_Left_Trackpad_Up", "Vive (L) Trackpad Up"), FKeyDetails::GamepadKey | FKeyDetails::NotBlueprintBindableKey, "Vive"));
+	EKeys::AddKey(FKeyDetails(ViveControllerKeys::Vive_Left_Trackpad_Down, LOCTEXT("Vive_Left_Trackpad_Down", "Vive (L) Trackpad Down"), FKeyDetails::GamepadKey | FKeyDetails::NotBlueprintBindableKey, "Vive"));
+	EKeys::AddKey(FKeyDetails(ViveControllerKeys::Vive_Left_Trackpad_Left, LOCTEXT("Vive_Left_Trackpad_Left", "Vive (L) Trackpad Left"), FKeyDetails::GamepadKey | FKeyDetails::NotBlueprintBindableKey, "Vive"));
+	EKeys::AddKey(FKeyDetails(ViveControllerKeys::Vive_Left_Trackpad_Right, LOCTEXT("Vive_Left_Trackpad_Right", "Vive (L) Trackpad Right"), FKeyDetails::GamepadKey | FKeyDetails::NotBlueprintBindableKey, "Vive"));
+
+	EKeys::AddKey(FKeyDetails(ViveControllerKeys::Vive_Right_Trackpad_Up, LOCTEXT("Vive_Right_Trackpad_Up", "Vive (R) Trackpad Up"), FKeyDetails::GamepadKey | FKeyDetails::NotBlueprintBindableKey, "Vive"));
+	EKeys::AddKey(FKeyDetails(ViveControllerKeys::Vive_Right_Trackpad_Down, LOCTEXT("Vive_Right_Trackpad_Down", "Vive (R) Trackpad Down"), FKeyDetails::GamepadKey | FKeyDetails::NotBlueprintBindableKey, "Vive"));
+	EKeys::AddKey(FKeyDetails(ViveControllerKeys::Vive_Right_Trackpad_Left, LOCTEXT("Vive_Right_Trackpad_Left", "Vive (R) Trackpad Left"), FKeyDetails::GamepadKey | FKeyDetails::NotBlueprintBindableKey, "Vive"));
+	EKeys::AddKey(FKeyDetails(ViveControllerKeys::Vive_Right_Trackpad_Right, LOCTEXT("Vive_Right_Trackpad_Right", "Vive (R) Trackpad Right"), FKeyDetails::GamepadKey | FKeyDetails::NotBlueprintBindableKey, "Vive"));
+#pragma endregion
+
+#pragma region WINDOWS MIXEDREALITY
+	EKeys::AddKey(FKeyDetails(WindowsMRKeys::MixedReality_Left_Trackpad_Up, LOCTEXT("MixedReality_Left_Trackpad_Up", "Mixed Reality (L) Trackpad Up"), FKeyDetails::GamepadKey | FKeyDetails::NotBlueprintBindableKey, "MixedReality"));
+	EKeys::AddKey(FKeyDetails(WindowsMRKeys::MixedReality_Left_Trackpad_Down, LOCTEXT("MixedReality_Left_Trackpad_Down", "Mixed Reality (L) Trackpad Down"), FKeyDetails::GamepadKey | FKeyDetails::NotBlueprintBindableKey, "MixedReality"));
+	EKeys::AddKey(FKeyDetails(WindowsMRKeys::MixedReality_Left_Trackpad_Left, LOCTEXT("MixedReality_Left_Trackpad_Left", "Mixed Reality (L) Trackpad Left"), FKeyDetails::GamepadKey | FKeyDetails::NotBlueprintBindableKey, "MixedReality"));
+	EKeys::AddKey(FKeyDetails(WindowsMRKeys::MixedReality_Left_Trackpad_Right, LOCTEXT("MixedReality_Left_Trackpad_Right", "Mixed Reality (L) Trackpad Right"), FKeyDetails::GamepadKey | FKeyDetails::NotBlueprintBindableKey, "MixedReality"));
+
+	EKeys::AddKey(FKeyDetails(WindowsMRKeys::MixedReality_Right_Trackpad_Up, LOCTEXT("MixedReality_Right_Trackpad_Up", "Mixed Reality (R) Trackpad Up"), FKeyDetails::GamepadKey | FKeyDetails::NotBlueprintBindableKey, "MixedReality"));
+	EKeys::AddKey(FKeyDetails(WindowsMRKeys::MixedReality_Right_Trackpad_Down, LOCTEXT("MixedReality_Right_Trackpad_Down", "Mixed Reality (R) Trackpad Down"), FKeyDetails::GamepadKey | FKeyDetails::NotBlueprintBindableKey, "MixedReality"));
+	EKeys::AddKey(FKeyDetails(WindowsMRKeys::MixedReality_Right_Trackpad_Left, LOCTEXT("MixedReality_Right_Trackpad_Left", "Mixed Reality (R) Trackpad Left"), FKeyDetails::GamepadKey | FKeyDetails::NotBlueprintBindableKey, "MixedReality"));
+	EKeys::AddKey(FKeyDetails(WindowsMRKeys::MixedReality_Right_Trackpad_Right, LOCTEXT("MixedReality_Right_Trackpad_Right", "Mixed Reality (R) Trackpad Right"), FKeyDetails::GamepadKey | FKeyDetails::NotBlueprintBindableKey, "MixedReality"));
+#pragma endregion
+
+#pragma region COSMOS KEYS
+	EKeys::AddMenuCategoryDisplayInfo("Cosmos", LOCTEXT("CosmosSubCategory", "HTC Cosmos"), TEXT("GraphEditor.PadEvent_16x"));
+
+	EKeys::AddKey(FKeyDetails(CosmosKeys::Cosmos_Left_X_Click, LOCTEXT("Cosmos_Left_X_Click", "Cosmos (L) X Press"), FKeyDetails::GamepadKey | FKeyDetails::NotBlueprintBindableKey, "Cosmos"));
+	EKeys::AddKey(FKeyDetails(CosmosKeys::Cosmos_Left_Y_Click, LOCTEXT("Cosmos_Left_Y_Click", "Cosmos (L) Y Press"), FKeyDetails::GamepadKey | FKeyDetails::NotBlueprintBindableKey, "Cosmos"));
+	EKeys::AddKey(FKeyDetails(CosmosKeys::Cosmos_Left_X_Touch, LOCTEXT("Cosmos_Left_X_Touch", "Cosmos (L) X Touch"), FKeyDetails::GamepadKey | FKeyDetails::NotBlueprintBindableKey, "Cosmos"));
+	EKeys::AddKey(FKeyDetails(CosmosKeys::Cosmos_Left_Y_Touch, LOCTEXT("Cosmos_Left_Y_Touch", "Cosmos (L) Y Touch"), FKeyDetails::GamepadKey | FKeyDetails::NotBlueprintBindableKey, "Cosmos"));
+	EKeys::AddKey(FKeyDetails(CosmosKeys::Cosmos_Left_Menu_Click, LOCTEXT("Cosmos_Left_Menu_Click", "Cosmos (L) Menu"), FKeyDetails::GamepadKey | FKeyDetails::NotBlueprintBindableKey, "Cosmos"));
+	EKeys::AddKey(FKeyDetails(CosmosKeys::Cosmos_Left_Grip_Click, LOCTEXT("Cosmos_Left_Grip_Click", "Cosmos (L) Grip"), FKeyDetails::GamepadKey | FKeyDetails::NotBlueprintBindableKey, "Cosmos"));
+	EKeys::AddKey(FKeyDetails(CosmosKeys::Cosmos_Left_Grip_Axis, LOCTEXT("Cosmos_Left_Grip_Axis", "Cosmos (L) Grip Axis"), FKeyDetails::GamepadKey | FKeyDetails::FloatAxis, "Cosmos"));
+	EKeys::AddKey(FKeyDetails(CosmosKeys::Cosmos_Left_Trigger_Click, LOCTEXT("Cosmos_Left_Trigger_Click", "Cosmos (L) Trigger"), FKeyDetails::GamepadKey | FKeyDetails::NotBlueprintBindableKey, "Cosmos"));
+	EKeys::AddKey(FKeyDetails(CosmosKeys::Cosmos_Left_Trigger_Axis, LOCTEXT("Cosmos_Left_Trigger_Axis", "Cosmos (L) Trigger Axis"), FKeyDetails::GamepadKey | FKeyDetails::FloatAxis, "Cosmos"));
+	EKeys::AddKey(FKeyDetails(CosmosKeys::Cosmos_Left_Trigger_Touch, LOCTEXT("Cosmos_Left_Trigger_Touch", "Cosmos (L) Trigger Touch"), FKeyDetails::GamepadKey | FKeyDetails::NotBlueprintBindableKey, "Cosmos"));
+	EKeys::AddKey(FKeyDetails(CosmosKeys::Cosmos_Left_Thumbstick_Vector, LOCTEXT("Cosmos_Left_Thumbstick_Vector", "Cosmos (L) Thumbstick Vector"), FKeyDetails::GamepadKey | FKeyDetails::VectorAxis, "Cosmos"));
+	EKeys::AddKey(FKeyDetails(CosmosKeys::Cosmos_Left_Thumbstick_X, LOCTEXT("Cosmos_Left_Thumbstick_X", "Cosmos (L) Thumbstick X"), FKeyDetails::GamepadKey | FKeyDetails::FloatAxis, "Cosmos"));
+	EKeys::AddKey(FKeyDetails(CosmosKeys::Cosmos_Left_Thumbstick_Y, LOCTEXT("Cosmos_Left_Thumbstick_Y", "Cosmos (L) Thumbstick Y"), FKeyDetails::GamepadKey | FKeyDetails::FloatAxis, "Cosmos"));
+	EKeys::AddKey(FKeyDetails(CosmosKeys::Cosmos_Left_Thumbstick_Click, LOCTEXT("Cosmos_Left_Thumbstick_Click", "Cosmos (L) Thumbstick"), FKeyDetails::GamepadKey | FKeyDetails::NotBlueprintBindableKey, "Cosmos"));
+	EKeys::AddKey(FKeyDetails(CosmosKeys::Cosmos_Left_Bumper_Click, LOCTEXT("Cosmos_Left_Bumper_Click", "Cosmos (L) Bumper"), FKeyDetails::GamepadKey | FKeyDetails::NotBlueprintBindableKey, "Cosmos"));
+	EKeys::AddKey(FKeyDetails(CosmosKeys::Cosmos_Left_Thumbstick_Touch, LOCTEXT("Cosmos_Left_Thumbstick_Touch", "Cosmos (L) Thumbstick Touch"), FKeyDetails::GamepadKey | FKeyDetails::NotBlueprintBindableKey, "Cosmos"));
+	EKeys::AddKey(FKeyDetails(CosmosKeys::Cosmos_Right_A_Click, LOCTEXT("Cosmos_Right_A_Click", "Cosmos (R) A Press"), FKeyDetails::GamepadKey | FKeyDetails::NotBlueprintBindableKey, "Cosmos"));
+	EKeys::AddKey(FKeyDetails(CosmosKeys::Cosmos_Right_B_Click, LOCTEXT("Cosmos_Right_B_Click", "Cosmos (R) B Press"), FKeyDetails::GamepadKey | FKeyDetails::NotBlueprintBindableKey, "Cosmos"));
+	EKeys::AddKey(FKeyDetails(CosmosKeys::Cosmos_Right_A_Touch, LOCTEXT("Cosmos_Right_A_Touch", "Cosmos (R) A Touch"), FKeyDetails::GamepadKey | FKeyDetails::NotBlueprintBindableKey, "Cosmos"));
+	EKeys::AddKey(FKeyDetails(CosmosKeys::Cosmos_Right_B_Touch, LOCTEXT("Cosmos_Right_B_Touch", "Cosmos (R) B Touch"), FKeyDetails::GamepadKey | FKeyDetails::NotBlueprintBindableKey, "Cosmos"));
+	EKeys::AddKey(FKeyDetails(CosmosKeys::Cosmos_Right_System_Click, LOCTEXT("Cosmos_Right_System_Click", "Cosmos (R) System"), FKeyDetails::GamepadKey | FKeyDetails::NotBlueprintBindableKey, "Cosmos"));
+	EKeys::AddKey(FKeyDetails(CosmosKeys::Cosmos_Right_Grip_Click, LOCTEXT("Cosmos_Right_Grip_Click", "Cosmos (R) Grip"), FKeyDetails::GamepadKey | FKeyDetails::NotBlueprintBindableKey, "Cosmos"));
+	EKeys::AddKey(FKeyDetails(CosmosKeys::Cosmos_Right_Grip_Axis, LOCTEXT("Cosmos_Right_Grip_Axis", "Cosmos (R) Grip Axis"), FKeyDetails::GamepadKey | FKeyDetails::FloatAxis, "Cosmos"));
+	EKeys::AddKey(FKeyDetails(CosmosKeys::Cosmos_Right_Trigger_Click, LOCTEXT("Cosmos_Right_Trigger_Click", "Cosmos (R) Trigger"), FKeyDetails::GamepadKey | FKeyDetails::NotBlueprintBindableKey, "Cosmos"));
+	EKeys::AddKey(FKeyDetails(CosmosKeys::Cosmos_Right_Trigger_Axis, LOCTEXT("Cosmos_Right_Trigger_Axis", "Cosmos (R) Trigger Axis"), FKeyDetails::GamepadKey | FKeyDetails::FloatAxis, "Cosmos"));
+	EKeys::AddKey(FKeyDetails(CosmosKeys::Cosmos_Right_Trigger_Touch, LOCTEXT("Cosmos_Right_Trigger_Touch", "Cosmos (R) Trigger Touch"), FKeyDetails::GamepadKey | FKeyDetails::NotBlueprintBindableKey, "Cosmos"));
+	EKeys::AddKey(FKeyDetails(CosmosKeys::Cosmos_Right_Thumbstick_Vector, LOCTEXT("Cosmos_Right_Thumbstick_Vector", "Cosmos (R) Thumbstick Vector"), FKeyDetails::GamepadKey | FKeyDetails::VectorAxis, "Cosmos"));
+	EKeys::AddKey(FKeyDetails(CosmosKeys::Cosmos_Right_Thumbstick_X, LOCTEXT("Cosmos_Right_Thumbstick_X", "Cosmos (R) Thumbstick X"), FKeyDetails::GamepadKey | FKeyDetails::FloatAxis, "Cosmos"));
+	EKeys::AddKey(FKeyDetails(CosmosKeys::Cosmos_Right_Thumbstick_Y, LOCTEXT("Cosmos_Right_Thumbstick_Y", "Cosmos (R) Thumbstick Y"), FKeyDetails::GamepadKey | FKeyDetails::FloatAxis, "Cosmos"));
+	EKeys::AddKey(FKeyDetails(CosmosKeys::Cosmos_Right_Thumbstick_Click, LOCTEXT("Cosmos_Right_Thumbstick_Click", "Cosmos (R) Thumbstick"), FKeyDetails::GamepadKey | FKeyDetails::NotBlueprintBindableKey, "Cosmos"));
+	EKeys::AddKey(FKeyDetails(CosmosKeys::Cosmos_Right_Thumbstick_Touch, LOCTEXT("Cosmos_Right_Thumbstick_Touch", "Cosmos (R) Thumbstick Touch"), FKeyDetails::GamepadKey | FKeyDetails::NotBlueprintBindableKey, "Cosmos"));
+	EKeys::AddKey(FKeyDetails(CosmosKeys::Cosmos_Right_Bumper_Click, LOCTEXT("Cosmos_Right_Bumper_Click", "Cosmos (R) Bumper"), FKeyDetails::GamepadKey | FKeyDetails::NotBlueprintBindableKey, "Cosmos"));
+
+	EKeys::AddKey(FKeyDetails(CosmosKeys::Cosmos_Left_Thumbstick_Up, LOCTEXT("Cosmos_Left_Thumbstick_Up", "Cosmos (L) Thumbstick Up"), FKeyDetails::GamepadKey | FKeyDetails::NotBlueprintBindableKey, "Cosmos"));
+	EKeys::AddKey(FKeyDetails(CosmosKeys::Cosmos_Left_Thumbstick_Down, LOCTEXT("Cosmos_Left_Thumbstick_Down", "Cosmos (L) Thumbstick Down"), FKeyDetails::GamepadKey | FKeyDetails::NotBlueprintBindableKey, "Cosmos"));
+	EKeys::AddKey(FKeyDetails(CosmosKeys::Cosmos_Left_Thumbstick_Left, LOCTEXT("Cosmos_Left_Thumbstick_Left", "Cosmos (L) Thumbstick Left"), FKeyDetails::GamepadKey | FKeyDetails::NotBlueprintBindableKey, "Cosmos"));
+	EKeys::AddKey(FKeyDetails(CosmosKeys::Cosmos_Left_Thumbstick_Right, LOCTEXT("Cosmos_Left_Thumbstick_Right", "Cosmos (L) Thumbstick Right"), FKeyDetails::GamepadKey | FKeyDetails::NotBlueprintBindableKey, "Cosmos"));
+
+	EKeys::AddKey(FKeyDetails(CosmosKeys::Cosmos_Right_Thumbstick_Up, LOCTEXT("Cosmos_Right_Thumbstick_Up", "Cosmos (R) Thumbstick Up"), FKeyDetails::GamepadKey | FKeyDetails::NotBlueprintBindableKey, "Cosmos"));
+	EKeys::AddKey(FKeyDetails(CosmosKeys::Cosmos_Right_Thumbstick_Down, LOCTEXT("Cosmos_Right_Thumbstick_Down", "Cosmos (R) Thumbstick Down"), FKeyDetails::GamepadKey | FKeyDetails::NotBlueprintBindableKey, "Cosmos"));
+	EKeys::AddKey(FKeyDetails(CosmosKeys::Cosmos_Right_Thumbstick_Left, LOCTEXT("Cosmos_Right_Thumbstick_Left", "Cosmos (R) Thumbstick Left"), FKeyDetails::GamepadKey | FKeyDetails::NotBlueprintBindableKey, "Cosmos"));
+	EKeys::AddKey(FKeyDetails(CosmosKeys::Cosmos_Right_Thumbstick_Right, LOCTEXT("Cosmos_Right_Thumbstick_Right", "Cosmos (R) Thumbstick Right"), FKeyDetails::GamepadKey | FKeyDetails::NotBlueprintBindableKey, "Cosmos"));
+#pragma endregion
+}
+
+#if WITH_EDITOR
+void FSteamVRInputDevice::RegenerateActionManifest()
+{
+	this->GenerateActionManifest(true, false, true, true);
+}
+
+void FSteamVRInputDevice::RegenerateControllerBindings()
+{
+	this->GenerateActionManifest(false, true, true, true);
+}
+
+bool FSteamVRInputDevice::GenerateAppManifest(FString ManifestPath, FString ProjectName, FString& OutAppKey, FString& OutAppManifestPath)
+{
+	// Set SteamVR AppKey
+	OutAppKey = (TEXT(APP_MANIFEST_PREFIX) + SanitizeString(GameProjectName) + TEXT(".") + ProjectName).ToLower();
+	EditorAppKey = FString(OutAppKey);
+
+	// Set Application Manifest Path - same directory where the action manifest will be
+	OutAppManifestPath = FPaths::ProjectConfigDir() / APP_MANIFEST_FILE;
+	IFileManager& FileManager = FFileManagerGeneric::Get();
+
+	// Create Application Manifest json objects
+	TSharedRef<FJsonObject> AppManifestObject = MakeShareable(new FJsonObject());
+	TArray<TSharedPtr<FJsonValue>> ManifestApps;
+
+	// Add current engine version being used as source
+	AppManifestObject->SetStringField("source", FString::Printf(TEXT("UE")));
+
+	// Define the application setting that will be registered with SteamVR
+	TArray<TSharedPtr<FJsonValue>> ManifestApp;
+
+	// Create Application Object 
+	TSharedRef<FJsonObject> ApplicationObject = MakeShareable(new FJsonObject());
+	TArray<FString> AppStringFields = { "app_key",  OutAppKey,
+										"launch_type", "url",
+										"url", "steam://launch/",
+										"action_manifest_path", *IFileManager::Get().ConvertToAbsolutePathForExternalAppForRead(*ManifestPath)
+	};
+	BuildJsonObject(AppStringFields, ApplicationObject);
+
+	// Create localization object
+	TSharedPtr<FJsonObject> LocStringsObject = MakeShareable(new FJsonObject());
+	TSharedRef<FJsonObject> AppNameObject = MakeShareable(new FJsonObject());
+	AppNameObject->SetStringField("name", GameProjectName + " [UE Editor]");
+	LocStringsObject->SetObjectField("en_us", AppNameObject);
+	ApplicationObject->SetObjectField("strings", LocStringsObject);
+
+	// Assemble the json app manifest
+	ManifestApps.Add(MakeShareable(new FJsonValueObject(ApplicationObject)));
+	AppManifestObject->SetArrayField(TEXT("applications"), ManifestApps);
+
+	// Serialize json app manifest
+	FString AppManifestString;
+	TSharedRef<TJsonWriter<>> JsonWriter = TJsonWriterFactory<>::Create(&AppManifestString);
+	FJsonSerializer::Serialize(AppManifestObject, JsonWriter);
+
+	// Save json as a UTF8 file
+	if (!FFileHelper::SaveStringToFile(AppManifestString, *OutAppManifestPath, FFileHelper::EEncodingOptions::ForceUTF8WithoutBOM))
+	{
+		UE_LOG(LogSteamVRInputDevice, Error, TEXT("Error trying to generate application manifest in: %s"), *OutAppManifestPath);
+		return false;
+	}
+
+	return true;
+}
+
+void FSteamVRInputDevice::ReloadActionManifest()
+{
+	if (VRSystem() && VRInput() && VRApplications())
+	{
+		// Set Action Manifest Path
+		const FString ManifestPath = FPaths::ProjectConfigDir() / CONTROLLER_BINDING_PATH / ACTION_MANIFEST;
+		UE_LOG(LogSteamVRInputDevice, Display, TEXT("Reloading Action Manifest in: %s"), *ManifestPath);
+			
+		// Load application manifest
+		FString AppManifestPath = FPaths::ProjectConfigDir() / APP_MANIFEST_FILE;
+		EVRApplicationError AppError = VRApplications()->AddApplicationManifest(TCHAR_TO_UTF8(*IFileManager::Get().ConvertToAbsolutePathForExternalAppForRead(*AppManifestPath)), true);
+		UE_LOG(LogSteamVRInputDevice, Display, TEXT("[STEAMVR INPUT] Registering Application Manifest %s : %s"), *AppManifestPath, *FString(UTF8_TO_TCHAR(VRApplications()->GetApplicationsErrorNameFromEnum(AppError))));
+		
+		// Get the App Process Id
+		uint32 AppProcessId = FPlatformProcess::GetCurrentProcessId();
+		
+		// Set SteamVR AppKey
+		FString AppFileName = FPaths::GetCleanFilename(FPlatformProcess::GetApplicationName(AppProcessId));
+		FString SteamVRAppKey = (TEXT(APP_MANIFEST_PREFIX) + SanitizeString(GameProjectName) + TEXT(".") + AppFileName).ToLower();
+		
+		// Set AppKey for this Editor Session
+		AppError = VRApplications()->IdentifyApplication(AppProcessId, TCHAR_TO_UTF8(*SteamVRAppKey));
+		UE_LOG(LogSteamVRInputDevice, Display, TEXT("[STEAMVR INPUT] Editor Application [%d][%s] identified to SteamVR: %s"), AppProcessId, *SteamVRAppKey, *FString(UTF8_TO_TCHAR(VRApplications()->GetApplicationsErrorNameFromEnum(AppError))));
+
+		// Set Action Manifest
+		EVRInputError InputError = VRInput()->SetActionManifestPath(TCHAR_TO_UTF8(*IFileManager::Get().ConvertToAbsolutePathForExternalAppForRead(*ManifestPath)));
+		GetInputError(InputError, FString(TEXT("Setting Action Manifest Path")));
+	}
+}
+#endif
+
+void FSteamVRInputDevice::GenerateControllerBindings(const FString& BindingsPath, TArray<FControllerType>& InOutControllerTypes, TArray<TSharedPtr<FJsonValue>>& DefaultBindings, TArray<FSteamVRInputAction>& InActionsArray, TArray<FInputMapping>& InInputMapping, bool bDeleteIfExists)
+{
+	// Create the bindings directory if it doesn't exist
+	IFileManager& FileManager = FFileManagerGeneric::Get();
+	if (!FileManager.DirectoryExists(*BindingsPath))
+	{
+		FileManager.MakeDirectory(*BindingsPath);
+	}
+
+	// Go through all supported controller types
+	for (auto& SupportedController : InOutControllerTypes)
+	{
+		// If there is no user-defined controller binding or it hasn't been auto-generated yet, generate it
+		if (!SupportedController.bIsGenerated || bDeleteIfExists)
+		{
+			// Creating bindings file
+			TSharedRef<FJsonObject> BindingsObject = MakeShareable(new FJsonObject());
+			BindingsObject->SetStringField(TEXT("name"), TEXT("Default bindings for ") + SupportedController.Description);
+			BindingsObject->SetStringField(TEXT("controller_type"), SupportedController.Name.ToString());
+
+			// Create Action Bindings in JSON Format
+			TArray<TSharedPtr<FJsonValue>> JsonValuesArray;
+			GenerateActionBindings(InInputMapping, JsonValuesArray, SupportedController);
+
+			// Ensure we also handle generic UE4 Motion Controllers
+			if (!SupportedController.Description.Contains(TEXT("Headset"), ESearchCase::IgnoreCase, ESearchDir::FromEnd))
+			{
+				FControllerType GenericController = FControllerType(TEXT("MotionController"), TEXT("MotionController"), TEXT("MotionController"));
+				GenerateActionBindings(InInputMapping, JsonValuesArray, SupportedController, true);
+			}
+
+			// Create Action Set
+			TSharedRef<FJsonObject> ActionSetJsonObject = MakeShareable(new FJsonObject());
+			ActionSetJsonObject->SetArrayField(TEXT("sources"), JsonValuesArray);
+
+			// Add tracker poses
+			if (SupportedController.KeyEquivalent.Equals(TEXT("SteamVR_Vive_Tracker")))
+			{
+				// Add Controller Pose Mappings
+				TArray<TSharedPtr<FJsonValue>> TrackerPoseArray;
+
+				// Add Pose: Special 1
+				TSharedRef<FJsonObject> Special1JsonObject = MakeShareable(new FJsonObject());
+				Special1JsonObject->SetStringField(TEXT("output"), TEXT(ACTION_PATH_SPECIAL_BACK_L));
+				Special1JsonObject->SetStringField(TEXT("path"), TEXT(ACTION_PATH_SPCL_BACK_LEFT));
+				Special1JsonObject->SetStringField(TEXT("requirement"), TEXT("optional"));
+
+				TSharedRef<FJsonValueObject> Special1JsonValueObject = MakeShareable(new FJsonValueObject(Special1JsonObject));
+				TrackerPoseArray.Add(Special1JsonValueObject);
+
+				// Add Pose: Special 2
+				TSharedRef<FJsonObject> Special2JsonObject = MakeShareable(new FJsonObject());
+				Special2JsonObject->SetStringField(TEXT("output"), TEXT(ACTION_PATH_SPECIAL_BACK_R));
+				Special2JsonObject->SetStringField(TEXT("path"), TEXT(ACTION_PATH_SPCL_BACK_RIGHT));
+				Special2JsonObject->SetStringField(TEXT("requirement"), TEXT("optional"));
+
+				TSharedRef<FJsonValueObject> Special2JsonObjectJsonValueObject = MakeShareable(new FJsonValueObject(Special2JsonObject));
+				TrackerPoseArray.Add(Special2JsonObjectJsonValueObject);
+
+				// Add Pose: Special 3
+				TSharedRef<FJsonObject> Special3JsonObject = MakeShareable(new FJsonObject());
+				Special3JsonObject->SetStringField(TEXT("output"), TEXT(ACTION_PATH_SPECIAL_FRONT_L));
+				Special3JsonObject->SetStringField(TEXT("path"), TEXT(ACTION_PATH_SPCL_FRONT_LEFT));
+				Special3JsonObject->SetStringField(TEXT("requirement"), TEXT("optional"));
+
+				TSharedRef<FJsonValueObject> Special3JsonValueObject = MakeShareable(new FJsonValueObject(Special3JsonObject));
+				TrackerPoseArray.Add(Special3JsonValueObject);
+
+				// Add Pose: Special 4
+				TSharedRef<FJsonObject> Special4JsonObject = MakeShareable(new FJsonObject());
+				Special4JsonObject->SetStringField(TEXT("output"), TEXT(ACTION_PATH_SPECIAL_FRONT_R));
+				Special4JsonObject->SetStringField(TEXT("path"), TEXT(ACTION_PATH_SPCL_FRONT_RIGHT));
+				Special4JsonObject->SetStringField(TEXT("requirement"), TEXT("optional"));
+
+				TSharedRef<FJsonValueObject> Special4JsonValueObject = MakeShareable(new FJsonValueObject(Special4JsonObject));
+				TrackerPoseArray.Add(Special4JsonValueObject);
+
+				// Add Pose: Special 5
+				TSharedRef<FJsonObject> Special5JsonObject = MakeShareable(new FJsonObject());
+				Special5JsonObject->SetStringField(TEXT("output"), TEXT(ACTION_PATH_SPECIAL_FRONTR_L));
+				Special5JsonObject->SetStringField(TEXT("path"), TEXT(ACTION_PATH_SPCL_FRONTR_LEFT));
+				Special5JsonObject->SetStringField(TEXT("requirement"), TEXT("optional"));
+
+				TSharedRef<FJsonValueObject> Special5JsonValueObject = MakeShareable(new FJsonValueObject(Special5JsonObject));
+				TrackerPoseArray.Add(Special5JsonValueObject);
+
+				// Add Pose: Special 6
+				TSharedRef<FJsonObject> Special6JsonObject = MakeShareable(new FJsonObject());
+				Special6JsonObject->SetStringField(TEXT("output"), TEXT(ACTION_PATH_SPECIAL_FRONTR_R));
+				Special6JsonObject->SetStringField(TEXT("path"), TEXT(ACTION_PATH_SPCL_FRONTR_RIGHT));
+				Special6JsonObject->SetStringField(TEXT("requirement"), TEXT("optional"));
+
+				TSharedRef<FJsonValueObject> Special6JsonValueObject = MakeShareable(new FJsonValueObject(Special6JsonObject));
+				TrackerPoseArray.Add(Special6JsonValueObject);
+
+				// Add Pose: Special 7
+				TSharedRef<FJsonObject> Special7JsonObject = MakeShareable(new FJsonObject());
+				Special7JsonObject->SetStringField(TEXT("output"), TEXT(ACTION_PATH_SPECIAL_PISTOL_L));
+				Special7JsonObject->SetStringField(TEXT("path"), TEXT(ACTION_PATH_SPCL_PISTOL_LEFT));
+				Special7JsonObject->SetStringField(TEXT("requirement"), TEXT("optional"));
+
+				TSharedRef<FJsonValueObject> Special7JsonValueObject = MakeShareable(new FJsonValueObject(Special7JsonObject));
+				TrackerPoseArray.Add(Special7JsonValueObject);
+
+				// Add Pose: Special 8
+				TSharedRef<FJsonObject> Special8JsonObject = MakeShareable(new FJsonObject());
+				Special8JsonObject->SetStringField(TEXT("output"), TEXT(ACTION_PATH_SPECIAL_PISTOL_R));
+				Special8JsonObject->SetStringField(TEXT("path"), TEXT(ACTION_PATH_SPCL_PISTOL_RIGHT));
+				Special8JsonObject->SetStringField(TEXT("requirement"), TEXT("optional"));
+
+				TSharedRef<FJsonValueObject> Special8JsonValueObject = MakeShareable(new FJsonValueObject(Special8JsonObject));
+				TrackerPoseArray.Add(Special8JsonValueObject);
+
+				// Add Controller Input Array To Action Set
+				ActionSetJsonObject->SetArrayField(TEXT("poses"), TrackerPoseArray);
+			}
+
+			// Do not add any default bindings for headsets and misc devices
+			if (!SupportedController.Description.Contains(TEXT("Headset"))
+				&& !SupportedController.KeyEquivalent.Equals(TEXT("SteamVR_Gamepads"))
+				&& !SupportedController.KeyEquivalent.Equals(TEXT("SteamVR_Vive_Tracker"))
+				)
+			{
+				// Add Controller Pose Mappings
+				TArray<TSharedPtr<FJsonValue>> ControllerPoseArray;
+	
+				// Add Pose: Left Controller 
+				TSharedRef<FJsonObject> ControllerLeftJsonObject = MakeShareable(new FJsonObject());
+				ControllerLeftJsonObject->SetStringField(TEXT("output"), TEXT(ACTION_PATH_CONTROLLER_LEFT));
+				ControllerLeftJsonObject->SetStringField(TEXT("path"), TEXT(ACTION_PATH_CONT_RAW_LEFT));
+				ControllerLeftJsonObject->SetStringField(TEXT("requirement"), TEXT("optional"));
+	
+				TSharedRef<FJsonValueObject> ControllerLeftJsonValueObject = MakeShareable(new FJsonValueObject(ControllerLeftJsonObject));
+				ControllerPoseArray.Add(ControllerLeftJsonValueObject);
+	
+				// Add Pose: Right Controller
+				TSharedRef<FJsonObject> ControllerRightJsonObject = MakeShareable(new FJsonObject());
+				ControllerRightJsonObject->SetStringField(TEXT("output"), TEXT(ACTION_PATH_CONTROLLER_RIGHT));
+				ControllerRightJsonObject->SetStringField(TEXT("path"), TEXT(ACTION_PATH_CONT_RAW_RIGHT));
+				ControllerLeftJsonObject->SetStringField(TEXT("requirement"), TEXT("optional"));
+	
+				TSharedRef<FJsonValueObject> ControllerRightJsonValueObject = MakeShareable(new FJsonValueObject(ControllerRightJsonObject));
+				ControllerPoseArray.Add(ControllerRightJsonValueObject);
+	
+				// Add Controller Input Array To Action Set
+				ActionSetJsonObject->SetArrayField(TEXT("poses"), ControllerPoseArray);
+	
+				// Add Skeleton Mappings
+				TArray<TSharedPtr<FJsonValue>> SkeletonValuesArray;
+	
+				// Add Skeleton: Left Hand 
+				TSharedRef<FJsonObject> SkeletonLeftJsonObject = MakeShareable(new FJsonObject());
+				SkeletonLeftJsonObject->SetStringField(TEXT("output"), TEXT(ACTION_PATH_SKELETON_LEFT));
+				SkeletonLeftJsonObject->SetStringField(TEXT("path"), TEXT(ACTION_PATH_USER_SKEL_LEFT));
+	
+				TSharedRef<FJsonValueObject> SkeletonLeftJsonValueObject = MakeShareable(new FJsonValueObject(SkeletonLeftJsonObject));
+				SkeletonValuesArray.Add(SkeletonLeftJsonValueObject);
+	
+				// Add Skeleton: Right Hand
+				TSharedRef<FJsonObject> SkeletonRightJsonObject = MakeShareable(new FJsonObject());
+				SkeletonRightJsonObject->SetStringField(TEXT("output"), TEXT(ACTION_PATH_SKELETON_RIGHT));
+				SkeletonRightJsonObject->SetStringField(TEXT("path"), TEXT(ACTION_PATH_USER_SKEL_RIGHT));
+	
+				TSharedRef<FJsonValueObject> SkeletonRightJsonValueObject = MakeShareable(new FJsonValueObject(SkeletonRightJsonObject));
+				SkeletonValuesArray.Add(SkeletonRightJsonValueObject);
+	
+				// Add Skeleton Input Array To Action Set
+				ActionSetJsonObject->SetArrayField(TEXT("skeleton"), SkeletonValuesArray);
+	
+				// Add Haptic Mappings
+				TArray<TSharedPtr<FJsonValue>> HapticValuesArray;
+	
+				// Add Haptic: Left Hand 
+				TSharedRef<FJsonObject> HapticLeftJsonObject = MakeShareable(new FJsonObject());
+				HapticLeftJsonObject->SetStringField(TEXT("output"), TEXT(ACTION_PATH_VIBRATE_LEFT));
+				HapticLeftJsonObject->SetStringField(TEXT("path"), TEXT(ACTION_PATH_USER_VIB_LEFT));
+	
+				TSharedRef<FJsonValueObject> HapticLeftJsonValueObject = MakeShareable(new FJsonValueObject(HapticLeftJsonObject));
+				HapticValuesArray.Add(HapticLeftJsonValueObject);
+	
+				// Add Haptic: Right Hand
+				TSharedRef<FJsonObject> HapticRightJsonObject = MakeShareable(new FJsonObject());
+				HapticRightJsonObject->SetStringField(TEXT("output"), TEXT(ACTION_PATH_VIBRATE_RIGHT));
+				HapticRightJsonObject->SetStringField(TEXT("path"), TEXT(ACTION_PATH_USER_VIB_RIGHT));
+	
+				TSharedRef<FJsonValueObject> HapticRightJsonValueObject = MakeShareable(new FJsonValueObject(HapticRightJsonObject));
+				HapticValuesArray.Add(HapticRightJsonValueObject);
+	
+				// Add Haptic Output Array To Action Set
+				ActionSetJsonObject->SetArrayField(TEXT("haptics"), HapticValuesArray);
+			}
+
+			// Create Bindings File that includes all Action Sets
+			TSharedRef<FJsonObject> BindingsJsonObject = MakeShareable(new FJsonObject());
+			BindingsJsonObject->SetObjectField(TEXT(ACTION_SET), ActionSetJsonObject);
+			BindingsObject->SetObjectField(TEXT("bindings"), BindingsJsonObject);
+
+			// Set description of Bindings file to the Project Name
+			BindingsObject->SetStringField(TEXT("description"), GameProjectName);
+
+			// Set Bindings File Path
+			FString BindingsFilePath = BindingsPath / SupportedController.Name.ToString() + TEXT(".json");
+
+			// Delete if it exists
+			if (FileManager.FileExists(*BindingsFilePath) && bDeleteIfExists)
+			{
+				FPlatformFileManager::Get().GetPlatformFile().DeleteFile(*BindingsFilePath);
+			}
+
+			// Save controller binding
+			FString OutputJsonString;
+			TSharedRef<TJsonWriter<>> JsonWriter = TJsonWriterFactory<>::Create(&OutputJsonString);
+			FJsonSerializer::Serialize(BindingsObject, JsonWriter);
+			FFileHelper::SaveStringToFile(OutputJsonString, *BindingsFilePath, FFileHelper::EEncodingOptions::ForceUTF8WithoutBOM);
+
+			// Create Controller Binding Object for this binding file
+			TSharedRef<FJsonObject> ControllerBindingObject = MakeShareable(new FJsonObject());
+			TArray<FString> ControllerStringFields = { "controller_type", *SupportedController.Name.ToString(),
+											 TEXT("binding_url"), *(SupportedController.Name.ToString() + TEXT(".json")) //*FileManager.ConvertToAbsolutePathForExternalAppForRead(*BindingsFilePath)
+			};
+			BuildJsonObject(ControllerStringFields, ControllerBindingObject);
+			DefaultBindings.Add(MakeShareable(new FJsonValueObject(ControllerBindingObject)));
+
+			// Tag this controller as generated
+			SupportedController.bIsGenerated = true;
+		}
+	}
+}
+
+void FSteamVRInputDevice::GenerateActionBindings(TArray<FInputMapping> &InInputMapping, TArray<TSharedPtr<FJsonValue>> &JsonValuesArray, FControllerType Controller, bool bIsGenericController)
+{
+	// Check for headsets
+	bool bIsHeadset = Controller.Description.Contains(TEXT("Headset"), ESearchCase::IgnoreCase, ESearchDir::FromEnd);
+	bool bIsGenericControllerCache = bIsGenericController;
+
+	// Process Key Input Mappings
+	for (FSteamVRInputKeyMapping SteamVRKeyInputMapping : SteamVRKeyInputMappings)
+	{
+
+		// Check if this is a generic UE motion controller key
+		bool bHasSteamVRInputs = false;
+		if (bIsGenericController)
+		{		
+			// Let's check if there're any SteamVR specific key that already exists for this action
+			for (FSteamVRInputKeyMapping SteamVRKeyInputMappingInner : SteamVRKeyInputMappings)
+			{
+				// Ensure the hmd proximity key is always generated
+				//if (SteamVRKeyInputMappingInner.InputKeyMapping.Key.GetFName().ToString().Contains(TEXT("SteamVR_HMD_Proximity")))
+				//{
+				//	UE_LOG(LogTemp, Warning, TEXT("*** Action includes a SteamVR_HMD_Proximity!"));
+				//	bIsGenericController = false;
+				//	break;
+				//}
+
+				// Check for generic controllers that have steamvr inputs already defined
+				if (SteamVRKeyInputMapping.InputKeyMapping.ActionName.ToString().Equals(SteamVRKeyInputMappingInner.InputKeyMapping.ActionName.ToString())
+					&& IsVRKey(SteamVRKeyInputMappingInner.InputKeyMapping.Key.GetFName())
+					&& !SteamVRKeyInputMappingInner.InputKeyMapping.Key.GetFName().ToString().Contains(TEXT("SteamVR_HMD_Proximity"))
+					)
+				{						
+					bHasSteamVRInputs = true;
+					break;
+				}
+				else
+				{
+					bHasSteamVRInputs = false;
+				}
+			}
+		}
+
+		if ((bIsGenericController && !bHasSteamVRInputs) || (!bIsGenericController && Controller.KeyEquivalent.Contains(TEXT("SteamVR")) && !SteamVRKeyInputMapping.ControllerName.Contains(TEXT("MotionController"))))
+		{
+			// Check this input mapping is of the correct controller type
+			if (!Controller.KeyEquivalent.Contains(SteamVRKeyInputMapping.ControllerName) 
+				&& !SteamVRKeyInputMapping.InputKeyMapping.Key.GetFName().ToString().Contains(TEXT("MotionController"))
+				&& !SteamVRKeyInputMapping.InputKeyMapping.Key.GetFName().ToString().Contains(TEXT("HMD_Proximity"))
+				)
+			{
+				//UE_LOG(LogTemp, Warning, TEXT("*** %s %s %s"), *Controller.KeyEquivalent, 
+				//	*SteamVRKeyInputMapping.InputKeyMapping.Key.GetFName().ToString(),
+				//	*SteamVRKeyInputMapping.ControllerName
+				//	);
+				continue;
+			}
+			else
+			{		
+				// Process the Key Mapping
+				FSteamVRInputState InputState;
+				FName CacheMode;
+				FString CacheType;
+				FString CachePath;
+
+				// Set Axis States
+				InputState.bIsAxis = false;
+				InputState.bIsAxis2 = false;
+				InputState.bIsAxis3 = false;
+
+				// Reset Dpad States
+				InputState.bIsDpadUp = false;
+				InputState.bIsDpadDown = false;
+				InputState.bIsDpadLeft = false;
+				InputState.bIsDpadRight = false;
+
+				// Set Input State
+				FString CurrentInputKeyName = SteamVRKeyInputMapping.InputKeyMapping.Key.ToString();
+				InputState.bIsTrigger = CurrentInputKeyName.Contains(TEXT("Trigger"), ESearchCase::CaseSensitive, ESearchDir::FromEnd);
+				InputState.bIsBumper = CurrentInputKeyName.Contains(TEXT("Bumper"), ESearchCase::CaseSensitive, ESearchDir::FromEnd);
+				InputState.bIsPress = CurrentInputKeyName.Contains(TEXT("Click"), ESearchCase::CaseSensitive, ESearchDir::FromEnd);
+				InputState.bIsThumbstick = InputState.bIsJoystick = CurrentInputKeyName.Contains(TEXT("Thumbstick"), ESearchCase::CaseSensitive, ESearchDir::FromEnd);
+				InputState.bIsTrackpad = CurrentInputKeyName.Contains(TEXT("Trackpad"), ESearchCase::CaseSensitive, ESearchDir::FromEnd);
+				InputState.bIsGrip = CurrentInputKeyName.Contains(TEXT("Grip"), ESearchCase::CaseSensitive, ESearchDir::FromEnd);
+				InputState.bIsLeft = CurrentInputKeyName.Contains(TEXT("_Left_"), ESearchCase::CaseSensitive, ESearchDir::FromEnd);
+				InputState.bIsFaceButton1 = CurrentInputKeyName.Contains(TEXT("FaceButton1"), ESearchCase::CaseSensitive, ESearchDir::FromEnd) ||
+					CurrentInputKeyName.Contains(TEXT("_A_"));
+				InputState.bIsFaceButton2 = CurrentInputKeyName.Contains(TEXT("FaceButton2"), ESearchCase::CaseSensitive, ESearchDir::FromEnd) ||
+					CurrentInputKeyName.Contains(TEXT("_B_"));
+				InputState.bIsAppMenu = CurrentInputKeyName.Contains(TEXT("_System_"));
+				InputState.bIsProximity = CurrentInputKeyName.Contains(TEXT("_HMD_Proximity"));
+				
+				// Only handle proximity sensor for headsets
+				if ((bIsHeadset && !InputState.bIsProximity)
+					|| (!bIsHeadset && InputState.bIsProximity)
+					)
+				{
+					continue;
+				}
+
+				// Handle Oculus Touch
+				InputState.bIsXButton = InputState.bIsYButton = false;
+				if (CurrentInputKeyName.Contains(TEXT("OculusTouch"))
+					|| CurrentInputKeyName.Contains(TEXT("Cosmos"))
+					)
+				{
+					// Check cap sense
+					FString ActualKeyName = CurrentInputKeyName.RightChop(19);
+					InputState.bIsCapSense = ActualKeyName.Contains(TEXT("_Touch"), ESearchCase::CaseSensitive, ESearchDir::FromEnd);
+
+					// Check for left X & Y buttons specific to Oculus Touch
+					InputState.bIsXButton = CurrentInputKeyName.Contains(TEXT("_X_Click")) || 
+						CurrentInputKeyName.Contains(TEXT("_X_Touch"));
+					InputState.bIsYButton = CurrentInputKeyName.Contains(TEXT("_Y_Click")) ||
+						CurrentInputKeyName.Contains(TEXT("_Y_Touch"));
+				}
+				else
+				{
+					// Set cap sense input state
+					InputState.bIsCapSense = CurrentInputKeyName.Contains(TEXT("CapSense"), ESearchCase::CaseSensitive, ESearchDir::FromEnd) ||
+						CurrentInputKeyName.Contains(TEXT("_Touch"), ESearchCase::CaseSensitive, ESearchDir::FromEnd);
+				}
+
+				// Check for DPad Keys
+				if (CurrentInputKeyName.Contains(TEXT("_Up"), ESearchCase::CaseSensitive, ESearchDir::FromEnd))
+				{
+					InputState.bIsDpadUp = true;
+					InputState.bIsDpadDown = false;
+					InputState.bIsDpadLeft = false;
+					InputState.bIsDpadRight = false;
+				}
+				else if (CurrentInputKeyName.Contains(TEXT("_Down"), ESearchCase::CaseSensitive, ESearchDir::FromEnd))
+				{
+					InputState.bIsDpadUp = false;
+					InputState.bIsDpadDown = true;
+					InputState.bIsDpadLeft = false;
+					InputState.bIsDpadRight = false;
+				}
+				else if (CurrentInputKeyName.Contains(TEXT("Trackpad_Left"), ESearchCase::CaseSensitive, ESearchDir::FromEnd) || CurrentInputKeyName.Contains(TEXT("Thumbstick_Left"), ESearchCase::CaseSensitive, ESearchDir::FromEnd))
+				{
+					InputState.bIsDpadUp = false;
+					InputState.bIsDpadDown = false;
+					InputState.bIsDpadLeft = true;
+					InputState.bIsDpadRight = false;
+				}
+				else if (CurrentInputKeyName.Contains(TEXT("Trackpad_Right"), ESearchCase::CaseSensitive, ESearchDir::FromEnd) || CurrentInputKeyName.Contains(TEXT("Thumbstick_Right"), ESearchCase::CaseSensitive, ESearchDir::FromEnd))
+				{
+					InputState.bIsDpadUp = false;
+					InputState.bIsDpadDown = false;
+					InputState.bIsDpadLeft = false;
+					InputState.bIsDpadRight = true;
+				}
+
+				// Handle Special Grip & Grab actions for supported controllers
+				if ((CurrentInputKeyName.Contains(TEXT("ValveIndex"), ESearchCase::IgnoreCase, ESearchDir::FromStart) 
+					|| CurrentInputKeyName.Contains(TEXT("Cosmos"), ESearchCase::IgnoreCase, ESearchDir::FromStart))
+					&& CurrentInputKeyName.Contains(TEXT("Pinch"), ESearchCase::IgnoreCase, ESearchDir::FromEnd))
+				{
+					InputState.bIsPinchGrab = true;
+					InputState.bIsGripGrab = false;
+					InputState.bIsGrip = false;
+					InputState.bIsAxis = false;
+				}
+				else if ((CurrentInputKeyName.Contains(TEXT("ValveIndex"), ESearchCase::IgnoreCase, ESearchDir::FromStart) 
+					|| CurrentInputKeyName.Contains(TEXT("Cosmos"), ESearchCase::IgnoreCase, ESearchDir::FromStart))
+					&& CurrentInputKeyName.Contains(TEXT("Grip"), ESearchCase::IgnoreCase, ESearchDir::FromEnd)
+					&& CurrentInputKeyName.Contains(TEXT("Grab"), ESearchCase::IgnoreCase, ESearchDir::FromEnd))
+				{
+					InputState.bIsGripGrab = true;
+					InputState.bIsPinchGrab = false;
+					InputState.bIsGrip = false;
+					InputState.bIsAxis = false;
+				}
+				else
+				{
+					InputState.bIsPinchGrab = false;
+					InputState.bIsGripGrab = false;
+				}
+
+				// Set Cache Mode
+				CacheMode = InputState.bIsTrigger || InputState.bIsGrip ? FName(TEXT("trigger")) : FName(TEXT("button"));
+				CacheMode = InputState.bIsPress ? FName(TEXT("button")) : CacheMode;
+				CacheMode = InputState.bIsTrackpad ? FName(TEXT("trackpad")) : CacheMode;
+				CacheMode = InputState.bIsJoystick ? FName(TEXT("joystick")) : CacheMode;
+				CacheMode = (InputState.bIsTrackpad || InputState.bIsJoystick) && !InputState.bIsAxis ? FName(TEXT("button")) : CacheMode;
+				CacheMode = InputState.bIsGrip ? FName(TEXT("button")) : CacheMode;
+				//CacheMode = InputState.bIsThumbstick ? FName(TEXT("button")) : CacheMode;
+				CacheMode = InputState.bIsPinchGrab || InputState.bIsGripGrab ? FName(TEXT("grab")) : CacheMode;
+
+				// Set Cache Path
+				if (InputState.bIsTrigger)
+				{
+					CachePath = InputState.bIsLeft ? FString(TEXT(ACTION_PATH_TRIGGER_LEFT)) : FString(TEXT(ACTION_PATH_TRIGGER_RIGHT));
+				}
+				else if (InputState.bIsBumper)
+				{
+					CachePath = InputState.bIsLeft ? FString(TEXT(ACTION_PATH_BUMPER_LEFT)) : FString(TEXT(ACTION_PATH_BUMPER_RIGHT));
+				}
+				//else if (InputState.bIsThumbstick)
+				//{
+				//	CachePath = InputState.bIsLeft ? FString(TEXT(ACTION_PATH_THUMBSTICK_LEFT)) : FString(TEXT(ACTION_PATH_THUMBSTICK_RIGHT));
+				//}
+				else if (InputState.bIsTrackpad)
+				{
+					CachePath = InputState.bIsLeft ? FString(TEXT(ACTION_PATH_TRACKPAD_LEFT)) : FString(TEXT(ACTION_PATH_TRACKPAD_RIGHT));
+				}
+				else if (InputState.bIsJoystick)
+				{
+					CachePath = InputState.bIsLeft ? FString(TEXT(ACTION_PATH_JOYSTICK_LEFT)) : FString(TEXT(ACTION_PATH_JOYSTICK_RIGHT));
+				}
+				else if (InputState.bIsGrip)
+				{
+					CachePath = InputState.bIsLeft ? FString(TEXT(ACTION_PATH_GRIP_LEFT)) : FString(TEXT(ACTION_PATH_GRIP_RIGHT));
+				}
+				else if (InputState.bIsFaceButton1)
+				{
+					CachePath = InputState.bIsLeft ? FString(TEXT(ACTION_PATH_BTN_A_LEFT)) : FString(TEXT(ACTION_PATH_BTN_A_RIGHT));
+				}
+				else if (InputState.bIsFaceButton2)
+				{
+					CachePath = InputState.bIsLeft ? FString(TEXT(ACTION_PATH_BTN_B_LEFT)) : FString(TEXT(ACTION_PATH_BTN_B_RIGHT));
+				}
+				else if (InputState.bIsXButton)
+				{
+					CachePath = FString(TEXT(ACTION_PATH_BTN_X_LEFT));
+				}
+				else if (InputState.bIsYButton)
+				{
+					CachePath = FString(TEXT(ACTION_PATH_BTN_Y_LEFT));
+				}
+				else if (InputState.bIsAppMenu)
+				{
+					CachePath = InputState.bIsLeft ? FString(TEXT(ACTION_PATH_APPMENU_LEFT)) : FString(TEXT(ACTION_PATH_APPMENU_RIGHT));
+				}
+
+				// Handle Special Actions
+				if (InputState.bIsPinchGrab)
+				{
+					CachePath = InputState.bIsLeft ? FString(TEXT(ACTION_PATH_PINCH_GRAB_LEFT)) : FString(TEXT(ACTION_PATH_PINCH_GRAB_RIGHT));
+				}
+				else if (InputState.bIsGripGrab)
+				{
+					CachePath = InputState.bIsLeft ? FString(TEXT(ACTION_PATH_GRIP_GRAB_LEFT)) : FString(TEXT(ACTION_PATH_GRIP_GRAB_RIGHT));
+				}
+				else if (InputState.bIsProximity)
+				{
+					CachePath = FString(TEXT(ACTION_PATH_HEAD_PROXIMITY));
+				}
+
+				// Override mode if Dpad
+				if (InputState.bIsDpadUp || InputState.bIsDpadDown || InputState.bIsDpadLeft || InputState.bIsDpadRight)
+				{
+					CacheMode = FName(TEXT("dpad"));
+				}
+
+				// Create Action Source
+				FActionSource ActionSource = FActionSource(CacheMode, CachePath);
+				TSharedRef<FJsonObject> ActionSourceJsonObject = MakeShareable(new FJsonObject());
+				ActionSourceJsonObject->SetStringField(TEXT("mode"), ActionSource.Mode.ToString());
+
+				// Set Action Path
+				if (!ActionSource.Path.IsEmpty())
+				{
+					ActionSourceJsonObject->SetStringField(TEXT("path"), ActionSource.Path);
+				}
+				else
+				{
+					continue;
+				}
+
+				// Add parameters if Dpad
+				if (InputState.bIsDpadUp || InputState.bIsDpadDown || InputState.bIsDpadLeft || InputState.bIsDpadRight)
+				{
+					// Create Submode
+					TSharedRef<FJsonObject> SubmodeJsonObject = MakeShareable(new FJsonObject());
+					SubmodeJsonObject->SetStringField(TEXT("sub_mode"), TEXT("click"));
+					
+					// Create Parameter
+					TSharedPtr<FJsonObject> ParametersJsonObject = MakeShareable(new FJsonObject());
+
+					// Set Submode as a parameter
+					ActionSourceJsonObject->SetObjectField(TEXT("parameters"), SubmodeJsonObject);
+				}
+
+				// Set Key Mappings
+				TSharedPtr<FJsonObject> ActionInputJsonObject = MakeShareable(new FJsonObject());
+
+				// Create Action Path
+				TSharedRef<FJsonObject> ActionPathJsonObject = MakeShareable(new FJsonObject());
+				ActionPathJsonObject->SetStringField(TEXT("output"), SteamVRKeyInputMapping.ActionNameWithPath);
+
+				// Set Cache Type
+				if (InputState.bIsAxis && InputState.bIsAxis2)
+				{
+					if (InputState.bIsGrip)
+					{
+						CacheType = FString(TEXT("force"));
+					}
+					else
+					{
+						if (CacheMode.IsEqual(TEXT("trigger")))
+						{
+							CacheType = FString(TEXT("pull"));
+						}
+						else
+						{
+							CacheType = FString(TEXT("position"));
+						}
+					}
+				}
+				else if (InputState.bIsAxis && !InputState.bIsAxis2)
+				{
+					if (InputState.bIsGrip)
+					{
+						CacheType = FString(TEXT("force"));
+					}
+					else if (!InputState.bIsThumbstick && !InputState.bIsTrackpad)
+					{
+						CacheType = FString(TEXT("pull"));
+					}
+					else
+					{
+						CacheType = "";
+					}
+				}
+				else if (!InputState.bIsAxis)
+				{
+					CacheType = (InputState.bIsCapSense) ? FString(TEXT("touch")) : FString(TEXT("click"));
+				}
+				else
+				{
+					CacheType = "";
+				}
+
+				// Handle Dpad values
+				if (InputState.bIsDpadUp)
+				{
+					CacheType = "north";
+				}
+				else if (InputState.bIsDpadDown)
+				{
+					CacheType = "south";
+				}
+				else if (InputState.bIsDpadLeft)
+				{
+					CacheType = "west";
+				}
+				else if (InputState.bIsDpadRight)
+				{
+					CacheType = "east";
+				}
+
+				// Handle special actions
+				if (InputState.bIsPinchGrab || InputState.bIsGripGrab)
+				{
+					CacheType = FString(TEXT("grab"));
+				}
+
+				// Special handling for axes
+				if ((CacheMode.IsEqual(TEXT("joystick")) || CacheMode.IsEqual(TEXT("trackpad")))
+					&& SteamVRKeyInputMapping.ActionNameWithPath.Right(4) == TEXT("axis")
+					&& CacheType == TEXT("position"))
+				{
+					CacheType = "";
+				}
+
+				// Override values in case of hmd proximity
+				if (InputState.bIsProximity)
+				{
+					CachePath = FString(TEXT(ACTION_PATH_HEAD_PROXIMITY));
+					CacheMode = FName(TEXT("button"));
+					CacheType = FString(TEXT("click"));;
+				}
+
+				if (!CacheType.IsEmpty())
+				{
+					// Set Action Input Type
+					ActionInputJsonObject->SetObjectField(CacheType, ActionPathJsonObject);
+
+					// Set Inputs
+					ActionSourceJsonObject->SetObjectField(TEXT("inputs"), ActionInputJsonObject);
+
+					// Add to Sources Array
+					TSharedRef<FJsonValueObject> JsonValueObject = MakeShareable(new FJsonValueObject(ActionSourceJsonObject));
+					JsonValuesArray.AddUnique(JsonValueObject);
+				}
+			}
+		}
+
+		bIsGenericController = bIsGenericControllerCache;
+	}
+
+	// Process Key Axis Mappings (skip headsets)
+	bIsGenericController = bIsGenericControllerCache;
+	if (!bIsHeadset)
+	{
+		for (FSteamVRAxisKeyMapping SteamVRAxisKeyMapping : SteamVRKeyAxisMappings)
+			{
+				// Check if this is a generic UE motion controller key
+				bool bHasSteamVRInputs = false;
+
+				if (bIsGenericController)
+				{
+					// Let's check if there're any SteamVR specific key that already exists for this action
+					for (FSteamVRAxisKeyMapping SteamVRKeyInputMappingInner : SteamVRKeyAxisMappings)
+					{
+						if (SteamVRAxisKeyMapping.InputAxisKeyMapping.AxisName.ToString().Equals(SteamVRKeyInputMappingInner.InputAxisKeyMapping.AxisName.ToString())
+							&& IsVRKey(SteamVRKeyInputMappingInner.InputAxisKeyMapping.Key.GetFName()))
+						{
+							bHasSteamVRInputs = true;
+							break;
+						}
+						else
+						{
+							bHasSteamVRInputs = false;
+						}
+					}
+				}
+
+				if ((bIsGenericController && !bHasSteamVRInputs) || (!bIsGenericController && Controller.KeyEquivalent.Contains(TEXT("SteamVR")) && !SteamVRAxisKeyMapping.ControllerName.Contains(TEXT("MotionController"))))
+				{
+					// Check this input mapping is of the correct controller type
+					if (!Controller.KeyEquivalent.Contains(SteamVRAxisKeyMapping.ControllerName) && !SteamVRAxisKeyMapping.InputAxisKeyMapping.Key.GetFName().ToString().Contains(TEXT("MotionController")))
+					{
+						continue;
+					}
+					else
+					{
+						// Process the Key Mapping
+						FSteamVRInputState InputState;
+						FName CacheMode;
+						FString CacheType;
+						FString CachePath;
+	
+						// Set Axis States
+						InputState.bIsAxis = false;
+						InputState.bIsAxis2 = false;
+						InputState.bIsAxis3 = false;
+						if (SteamVRAxisKeyMapping.ActionName.Contains(TEXT("_axis2d")))
+						{
+							InputState.bIsAxis2 = true;
+						}
+						else if (SteamVRAxisKeyMapping.ActionName.Contains(TEXT("_axis3d")))
+						{
+							InputState.bIsAxis3 = true;
+						}
+						else if (SteamVRAxisKeyMapping.ActionName.Contains(TEXT(" axis")))
+						{
+							InputState.bIsAxis = true;
+						}
+	
+						// Set Input State
+						FString CurrentInputKeyName = SteamVRAxisKeyMapping.InputAxisKeyMapping.Key.ToString();
+						InputState.bIsTrigger = CurrentInputKeyName.Contains(TEXT("Trigger"), ESearchCase::CaseSensitive, ESearchDir::FromEnd);
+						InputState.bIsBumper = CurrentInputKeyName.Contains(TEXT("Bumper"), ESearchCase::CaseSensitive, ESearchDir::FromEnd);
+						InputState.bIsThumbstick = InputState.bIsJoystick = CurrentInputKeyName.Contains(TEXT("Thumbstick"), ESearchCase::CaseSensitive, ESearchDir::FromEnd);
+						InputState.bIsTrackpad = CurrentInputKeyName.Contains(TEXT("Trackpad"), ESearchCase::CaseSensitive, ESearchDir::FromEnd);
+						InputState.bIsGrip = CurrentInputKeyName.Contains(TEXT("Grip"), ESearchCase::CaseSensitive, ESearchDir::FromEnd);
+						InputState.bIsLeft = CurrentInputKeyName.Contains(TEXT("_Left_"), ESearchCase::CaseSensitive, ESearchDir::FromEnd);
+						InputState.bIsFaceButton1 = CurrentInputKeyName.Contains(TEXT("FaceButton1"), ESearchCase::CaseSensitive, ESearchDir::FromEnd) ||
+							CurrentInputKeyName.Contains(TEXT("_A_"));
+						InputState.bIsFaceButton2 = CurrentInputKeyName.Contains(TEXT("FaceButton2"), ESearchCase::CaseSensitive, ESearchDir::FromEnd) ||
+							CurrentInputKeyName.Contains(TEXT("_B_"));
+	
+						// Thumbstick vs Joystick (to conform with new UE naming scheme
+						if (CurrentInputKeyName.Contains(TEXT("ValveIndex")))
+						{
+							if (InputState.bIsThumbstick)
+							{
+								InputState.bIsThumbstick = true;
+								InputState.bIsJoystick = false;
+							}
+						}
+						else
+						{
+							if (InputState.bIsThumbstick)
+							{
+								InputState.bIsThumbstick = false;
+								InputState.bIsJoystick = true;
+							}
+						}
+
+						// Handle Oculus Touch
+						InputState.bIsXButton = InputState.bIsYButton = false;
+						if (CurrentInputKeyName.Contains(TEXT("OculusTouch")))
+						{
+							// Check cap sense
+							FString OculusKeyName = CurrentInputKeyName.RightChop(20);
+							InputState.bIsCapSense = OculusKeyName.Contains(TEXT("_Touch"), ESearchCase::CaseSensitive, ESearchDir::FromEnd);
+	
+							// Check for left X & Y buttons specific to Oculus Touch
+							InputState.bIsXButton = CurrentInputKeyName.Contains(TEXT("_X_Click")) ||
+								CurrentInputKeyName.Contains(TEXT("_X_Touch"));
+							InputState.bIsYButton = CurrentInputKeyName.Contains(TEXT("_Y_Click")) ||
+								CurrentInputKeyName.Contains(TEXT("_Y_Touch"));
+						}
+						else
+						{
+							// Set cap sense input state
+							InputState.bIsCapSense = CurrentInputKeyName.Contains(TEXT("CapSense"), ESearchCase::CaseSensitive, ESearchDir::FromEnd) ||
+								CurrentInputKeyName.Contains(TEXT("_Touch"), ESearchCase::CaseSensitive, ESearchDir::FromEnd);
+						}
+	
+						// Handle Special Actions for Knuckles Keys
+						if ((CurrentInputKeyName.Contains(TEXT("ValveIndex"), ESearchCase::IgnoreCase, ESearchDir::FromStart)
+							|| CurrentInputKeyName.Contains(TEXT("Cosmos"), ESearchCase::IgnoreCase, ESearchDir::FromStart)) 
+							&& CurrentInputKeyName.Contains(TEXT("Pinch"), ESearchCase::IgnoreCase, ESearchDir::FromEnd))
+						{
+							InputState.bIsPinchGrab = true;
+							InputState.bIsGripGrab = false;
+							InputState.bIsGrip = false;
+							InputState.bIsAxis = false;
+						}
+						else if ((CurrentInputKeyName.Contains(TEXT("ValveIndex"), ESearchCase::IgnoreCase, ESearchDir::FromStart)
+							|| CurrentInputKeyName.Contains(TEXT("Cosmos"), ESearchCase::IgnoreCase, ESearchDir::FromStart))
+							&& CurrentInputKeyName.Contains(TEXT("Grip"), ESearchCase::IgnoreCase, ESearchDir::FromEnd) 
+							&& CurrentInputKeyName.Contains(TEXT("Grab"), ESearchCase::IgnoreCase, ESearchDir::FromEnd))
+						{
+							InputState.bIsGripGrab = true;
+							InputState.bIsPinchGrab = false;
+							InputState.bIsGrip = false;
+							InputState.bIsAxis = false;
+						}
+						else
+						{
+							InputState.bIsPinchGrab = false;
+							InputState.bIsGripGrab = false;
+						}
+	
+						// Set Cache Mode
+						CacheMode = InputState.bIsTrigger || InputState.bIsGrip ? FName(TEXT("trigger")) : FName(TEXT("button"));
+						CacheMode = InputState.bIsTrackpad ? FName(TEXT("trackpad")) : CacheMode;
+						CacheMode = InputState.bIsJoystick ? FName(TEXT("joystick")) : CacheMode;
+						CacheMode = InputState.bIsGrip ? FName(TEXT("force_sensor")) : CacheMode;
+						CacheMode = InputState.bIsThumbstick ? FName(TEXT("joystick")) : CacheMode;
+						CacheMode = InputState.bIsPinchGrab || InputState.bIsGripGrab ? FName(TEXT("grab")) : CacheMode;
+	
+						// If key being mapped is not an axis key (hardware-wise), set mode as an analog action (scalar_constant to 1.0f)
+						// [REDACTED_PROFILE_URL]/steamvr_unreal_plugin/issues/12
+						if (!SteamVRAxisKeyMapping.InputAxisKeyMapping.Key.IsFloatAxis())
+						{
+							CacheMode = FName(TEXT("scalar_constant"));
+						}
+	
+						// Set Cache Path
+						if (InputState.bIsTrigger)
+						{
+							CachePath = InputState.bIsLeft ? FString(TEXT(ACTION_PATH_TRIGGER_LEFT)) : FString(TEXT(ACTION_PATH_TRIGGER_RIGHT));
+						}
+						else if (InputState.bIsBumper)
+						{
+							CachePath = InputState.bIsLeft ? FString(TEXT(ACTION_PATH_BUMPER_LEFT)) : FString(TEXT(ACTION_PATH_BUMPER_RIGHT));
+						}
+						else if (InputState.bIsThumbstick)
+						{
+							CachePath = InputState.bIsLeft ? FString(TEXT(ACTION_PATH_THUMBSTICK_LEFT)) : FString(TEXT(ACTION_PATH_THUMBSTICK_RIGHT));
+						}
+						else if (InputState.bIsTrackpad)
+						{
+							CachePath = InputState.bIsLeft ? FString(TEXT(ACTION_PATH_TRACKPAD_LEFT)) : FString(TEXT(ACTION_PATH_TRACKPAD_RIGHT));
+						}
+						else if (InputState.bIsJoystick)
+						{
+							CachePath = InputState.bIsLeft ? FString(TEXT(ACTION_PATH_JOYSTICK_LEFT)) : FString(TEXT(ACTION_PATH_JOYSTICK_RIGHT));
+						}
+						else if (InputState.bIsGrip)
+						{
+							CachePath = InputState.bIsLeft ? FString(TEXT(ACTION_PATH_GRIP_LEFT)) : FString(TEXT(ACTION_PATH_GRIP_RIGHT));
+
+							// For controllers without force sensor support, use trigger value mode
+							if (!CurrentInputKeyName.Contains(TEXT("ValveIndex"), ESearchCase::IgnoreCase, ESearchDir::FromStart)
+								&& CurrentInputKeyName.Contains(TEXT("Axis"), ESearchCase::IgnoreCase, ESearchDir::FromEnd)
+								)
+							{
+								CacheMode = FName(TEXT("trigger"));
+							}
+						}
+						else if (InputState.bIsFaceButton1)
+						{
+							CachePath = InputState.bIsLeft ? FString(TEXT(ACTION_PATH_BTN_A_LEFT)) : FString(TEXT(ACTION_PATH_BTN_A_RIGHT));
+						}
+						else if (InputState.bIsFaceButton2)
+						{
+							CachePath = InputState.bIsLeft ? FString(TEXT(ACTION_PATH_BTN_B_LEFT)) : FString(TEXT(ACTION_PATH_BTN_B_RIGHT));
+						}
+						else if (InputState.bIsXButton)
+						{
+							CachePath = FString(TEXT(ACTION_PATH_BTN_X_LEFT));
+						}
+						else if (InputState.bIsYButton)
+						{
+							CachePath = FString(TEXT(ACTION_PATH_BTN_Y_LEFT));
+						}
+	
+						// Handle Special Actions
+						if (InputState.bIsPinchGrab)
+						{
+							CachePath = InputState.bIsLeft ? FString(TEXT(ACTION_PATH_PINCH_GRAB_LEFT)) : FString(TEXT(ACTION_PATH_PINCH_GRAB_RIGHT));
+						}
+						else if (InputState.bIsGripGrab)
+						{
+							CachePath = InputState.bIsLeft ? FString(TEXT(ACTION_PATH_GRIP_GRAB_LEFT)) : FString(TEXT(ACTION_PATH_GRIP_GRAB_RIGHT));
+						}
+	
+						// Create Action Source
+						FActionSource ActionSource = FActionSource(CacheMode, CachePath);
+						TSharedRef<FJsonObject> ActionSourceJsonObject = MakeShareable(new FJsonObject());
+						ActionSourceJsonObject->SetStringField(TEXT("mode"), ActionSource.Mode.ToString());
+	
+						// Set Action Path
+						if (!ActionSource.Path.IsEmpty())
+						{
+							ActionSourceJsonObject->SetStringField(TEXT("path"), ActionSource.Path);
+						}
+						else
+						{
+							continue;
+						}
+	
+						// Set Key Mappings
+						TSharedPtr<FJsonObject> ActionInputJsonObject = MakeShareable(new FJsonObject());
+	
+						// Create Action Path
+						TSharedRef<FJsonObject> ActionPathJsonObject = MakeShareable(new FJsonObject());
+						ActionPathJsonObject->SetStringField(TEXT("output"), SteamVRAxisKeyMapping.ActionNameWithPath);
+	
+						// Set Cache Type
+						if (CacheMode.IsEqual(TEXT("scalar_constant")))
+						{
+							CacheType = FString(TEXT("value"));
+						}
+						else if (InputState.bIsAxis && InputState.bIsAxis2)
+						{
+							if (InputState.bIsGrip)
+							{
+								if (CacheMode == FName("trigger"))
+								{
+									CacheType = FString(TEXT("pull"));
+								}
+								else
+								{
+									CacheType = FString(TEXT("force"));
+								}
+							}
+							else
+							{
+								if (CacheMode.IsEqual(TEXT("trigger")))
+								{
+									CacheType = FString(TEXT("pull"));
+								}
+								else
+								{
+									CacheType = FString(TEXT("position"));
+								}
+							}
+						}
+						else if (InputState.bIsAxis && !InputState.bIsAxis2)
+						{
+							if (InputState.bIsGrip)
+							{
+								if (CacheMode == FName("trigger"))
+								{
+									CacheType = FString(TEXT("pull"));
+								}
+								else
+								{
+									CacheType = FString(TEXT("force"));
+								}
+							}
+							else if (!InputState.bIsThumbstick && !InputState.bIsTrackpad)
+							{
+								CacheType = FString(TEXT("pull"));
+							}
+							else
+							{
+								CacheType = "";
+							}
+						}
+						else if (InputState.bIsAxis2)
+						{
+							CacheType = FString(TEXT("position"));
+						}
+						else if (!InputState.bIsAxis)
+						{
+							CacheType = (InputState.bIsCapSense) ? FString(TEXT("touch")) : CacheType;
+						}
+						else
+						{
+							CacheType = "";
+						}
+	
+						// Handle special actions
+						if (InputState.bIsPinchGrab || InputState.bIsGripGrab)
+						{
+							CacheType = FString(TEXT("grab"));
+						}
+	
+						if (!CacheType.IsEmpty() && !bIsHeadset)
+						{
+							// Set Action Input Type
+							ActionInputJsonObject->SetObjectField(CacheType, ActionPathJsonObject);
+	
+							// Set Inputs
+							ActionSourceJsonObject->SetObjectField(TEXT("inputs"), ActionInputJsonObject);
+	
+							// Add to Sources Array
+							TSharedRef<FJsonValueObject> JsonValueObject = MakeShareable(new FJsonValueObject(ActionSourceJsonObject));
+							JsonValuesArray.AddUnique(JsonValueObject);
+						}
+					}
+				}
+			}
+	}
+}
+
+void FSteamVRInputDevice::GenerateActionManifest(bool GenerateActions, bool GenerateBindings, bool RegisterApp, bool DeleteIfExists)
+{
+	// Set Action Manifest Path
+	const FString ManifestPath = FPaths::ProjectConfigDir() / CONTROLLER_BINDING_PATH / ACTION_MANIFEST;
+	UE_LOG(LogSteamVRInputDevice, Display, TEXT("Action Manifest Path: %s"), *ManifestPath);
+
+	// Create Action Manifest json object
+	TSharedRef<FJsonObject> ActionManifestObject = MakeShareable(new FJsonObject());
+	TArray<FString> LocalizationFields = {"language_tag", "en_us"};
+
+	// Set where to look for controller binding files and prepare file manager
+	const FString ControllerBindingsPath = FPaths::ProjectConfigDir() / CONTROLLER_BINDING_PATH;
+	UE_LOG(LogSteamVRInputDevice, Display, TEXT("Controller Bindings Path: %s"), *ControllerBindingsPath);
+	IFileManager& FileManager = FFileManagerGeneric::Get();
+
+	// Define Controller Types supported by SteamVR
+	TArray<TSharedPtr<FJsonValue>> ControllerBindings;
+	ControllerTypes.Empty();
+	ControllerTypes.Emplace(FControllerType(TEXT("knuckles"), TEXT("ValveIndex"), TEXT("SteamVR_ValveIndex")));
+	ControllerTypes.Emplace(FControllerType(TEXT("vive_controller"), TEXT("Vive"), TEXT("SteamVR_Vive")));
+	ControllerTypes.Emplace(FControllerType(TEXT("vive_cosmos_controller"), TEXT("Cosmos"), TEXT("SteamVR_Cosmos")));
+	ControllerTypes.Emplace(FControllerType(TEXT("oculus_touch"), TEXT("OculusTouch"), TEXT("SteamVR_OculusTouch")));
+	ControllerTypes.Emplace(FControllerType(TEXT("holographic_controller"), TEXT("MixedReality"), TEXT("SteamVR_MixedReality")));
+	
+	ControllerTypes.Emplace(FControllerType(TEXT("indexhmd"), TEXT("Valve Index Headset"), TEXT("SteamVR_Valve_Index_Headset")));
+	ControllerTypes.Emplace(FControllerType(TEXT("vive"), TEXT("Vive Headset"), TEXT("SteamVR_Vive_Headset")));
+	ControllerTypes.Emplace(FControllerType(TEXT("vive_pro"), TEXT("Vive Pro Headset"), TEXT("SteamVR_Vive_Pro_Headset")));
+	ControllerTypes.Emplace(FControllerType(TEXT("rift"), TEXT("Rift Headset"), TEXT("SteamVR_Rift_Headset")));
+
+	ControllerTypes.Emplace(FControllerType(TEXT("vive_tracker_camera"), TEXT("Vive Trackers"), TEXT("SteamVR_Vive_Tracker")));
+	ControllerTypes.Emplace(FControllerType(TEXT("gamepad"), TEXT("Gamepads"), TEXT("SteamVR_Gamepads")));
+	
+#pragma region ACTIONS
+	// Clear Actions cache
+	Actions.Empty();
+
+	// Setup Input Mappings cache
+	TArray<FInputMapping> InputMappings;
+	TArray<FName> UniqueInputs;
+
+	// Set Input Settings
+	auto InputSettings = GetDefault<UInputSettings>();
+
+	// Check if this project have input settings
+	if (InputSettings->IsValidLowLevelFast())
+	{
+		// Process all actions in this project (if any)
+		TArray<TSharedPtr<FJsonValue>> InputActionsArray;
+
+		// Setup cache for actions
+		TArray<FString> UniqueActions;
+
+		// Controller poses
+		{
+			FString ConstActionPath = FString(TEXT(ACTION_PATH_CONTROLLER_LEFT));
+			Actions.Add(FSteamVRInputAction(ConstActionPath, EActionType::Pose, false,
+				FName(TEXT("Left Controller [Pose]")), FString(TEXT(ACTION_PATH_CONT_RAW_LEFT))));
+		}
+		{
+			FString ConstActionPath = FString(TEXT(ACTION_PATH_CONTROLLER_RIGHT));
+			Actions.Add(FSteamVRInputAction(ConstActionPath, EActionType::Pose, false,
+				FName(TEXT("Right Controller [Pose]")), FString(TEXT(ACTION_PATH_CONT_RAW_RIGHT))));
+		}
+
+		// Other poses
+		{
+			FString ConstActionPath = FString(TEXT(ACTION_PATH_SPECIAL_BACK_L));
+			Actions.Add(FSteamVRInputAction(ConstActionPath, EActionType::Pose, false,
+				FName(TEXT("Special 1 [Tracker]")), FString(TEXT(ACTION_PATH_SPCL_BACK_LEFT))));
+		}
+		{
+			FString ConstActionPath = FString(TEXT(ACTION_PATH_SPECIAL_BACK_R));
+			Actions.Add(FSteamVRInputAction(ConstActionPath, EActionType::Pose, false,
+				FName(TEXT("Special 2 [Tracker]")), FString(TEXT(ACTION_PATH_SPCL_BACK_RIGHT))));
+		}
+		{
+			FString ConstActionPath = FString(TEXT(ACTION_PATH_SPECIAL_FRONT_L));
+			Actions.Add(FSteamVRInputAction(ConstActionPath, EActionType::Pose, false,
+				FName(TEXT("Special 3 [Tracker]")), FString(TEXT(ACTION_PATH_SPCL_FRONT_LEFT))));
+		}
+		{
+			FString ConstActionPath = FString(TEXT(ACTION_PATH_SPECIAL_FRONT_R));
+			Actions.Add(FSteamVRInputAction(ConstActionPath, EActionType::Pose, false,
+				FName(TEXT("Special 4 [Tracker]")), FString(TEXT(ACTION_PATH_SPCL_FRONT_RIGHT))));
+		}
+		{
+			FString ConstActionPath = FString(TEXT(ACTION_PATH_SPECIAL_FRONTR_L));
+			Actions.Add(FSteamVRInputAction(ConstActionPath, EActionType::Pose, false,
+				FName(TEXT("Special 5 [Tracker]")), FString(TEXT(ACTION_PATH_SPCL_FRONTR_LEFT))));
+		}
+		{
+			FString ConstActionPath = FString(TEXT(ACTION_PATH_SPECIAL_FRONTR_R));
+			Actions.Add(FSteamVRInputAction(ConstActionPath, EActionType::Pose, false,
+				FName(TEXT("Special 6 [Tracker]")), FString(TEXT(ACTION_PATH_SPCL_FRONTR_RIGHT))));
+		}
+		{
+			FString ConstActionPath = FString(TEXT(ACTION_PATH_SPECIAL_PISTOL_L));
+			Actions.Add(FSteamVRInputAction(ConstActionPath, EActionType::Pose, false,
+				FName(TEXT("Special 7 [Tracker]")), FString(TEXT(ACTION_PATH_SPCL_PISTOL_LEFT))));
+		}
+		{
+			FString ConstActionPath = FString(TEXT(ACTION_PATH_SPECIAL_PISTOL_R));
+			Actions.Add(FSteamVRInputAction(ConstActionPath, EActionType::Pose, false,
+				FName(TEXT("Special 8 [Tracker]")), FString(TEXT(ACTION_PATH_SPCL_PISTOL_RIGHT))));
+		}
+
+		// Skeletal Data
+		{
+			FString ConstActionPath = FString(TEXT(ACTION_PATH_SKELETON_LEFT));
+			Actions.Add(FSteamVRInputAction(ConstActionPath, EActionType::Skeleton, false,
+				FName(TEXT("Skeleton (Left)")), FString(TEXT(ACTION_PATH_SKEL_HAND_LEFT))));
+		}
+		{
+			FString ConstActionPath = FString(TEXT(ACTION_PATH_SKELETON_RIGHT));
+			Actions.Add(FSteamVRInputAction(ConstActionPath, EActionType::Skeleton, false,
+				FName(TEXT("Skeleton (Right)")), FString(TEXT(ACTION_PATH_SKEL_HAND_RIGHT))));
+		}
+
+		// Open console
+		{
+			const FKey* ConsoleKey = InputSettings->ConsoleKeys.FindByPredicate([](FKey& Key) { return Key.IsValid(); });
+			if (ConsoleKey != nullptr)
+			{
+				Actions.Add(FSteamVRInputAction(FString(TEXT(ACTION_PATH_OPEN_CONSOLE)), FName(TEXT("Open Console")), false, ConsoleKey->GetFName(), false));
+			}
+		}
+
+		// Haptics
+		{
+			FString ConstActionPath = FString(TEXT(ACTION_PATH_VIBRATE_LEFT));
+			Actions.Add(FSteamVRInputAction(ConstActionPath, EActionType::Vibration, false, FName(TEXT("Haptic (Left)"))));
+		}
+		{
+			FString ConstActionPath = FString(TEXT(ACTION_PATH_VIBRATE_RIGHT));
+			Actions.Add(FSteamVRInputAction(ConstActionPath, EActionType::Vibration, false, FName(TEXT("Haptic (Right)"))));
+		}
+
+		// Add base actions to the action manifest
+		ActionManifestObject->SetArrayField(TEXT("actions"), InputActionsArray);
+
+		// Add project's input key mappings to SteamVR's Input Actions
+		ProcessKeyInputMappings(InputSettings, UniqueInputs);
+
+		// Add project's input axis mappings to SteamVR's Input Actions
+		ProcessKeyAxisMappings(InputSettings, UniqueInputs);
+
+		// Reorganize all unique inputs to SteamVR style Input-to-Actions association
+		for (FName UniqueInput : UniqueInputs)
+		{
+			// Create New Input Mapping from Unique Input Key
+			FInputMapping NewInputMapping = FInputMapping();
+			FInputMapping NewAxisMapping = FInputMapping();
+			NewInputMapping.InputKey = UniqueInput;
+			NewAxisMapping.InputKey = UniqueInput;
+
+			// Go through all the project actions
+			for (FSteamVRInputAction& Action : Actions)
+			{
+				//UE_LOG(LogTemp, Warning, TEXT("Action: %s Type: [%s]"), *Action.KeyX.ToString(), *Action.GetActionTypeName());
+
+				// Check for boolean/digital input
+				if (Action.Type == EActionType::Boolean)
+				{
+					// Set Key Actions Linked To This Input Key
+					TArray<FInputActionKeyMapping> ActionKeyMappings;
+					FindActionMappings(InputSettings, Action.Name, ActionKeyMappings);
+					for (FInputActionKeyMapping ActionKeyMapping : ActionKeyMappings)
+					{
+						if (UniqueInput.IsEqual(ActionKeyMapping.Key.GetFName()))
+						{
+							NewInputMapping.Actions.AddUnique(Action.Path);
+						}
+					}
+				}
+
+				// Check for axes/analog input
+				if (Action.Type == EActionType::Vector1 || Action.Type == EActionType::Vector2 || Action.Type == EActionType::Vector3)
+				{
+					// Set Axis Actions Linked To This Input Key
+					FString ActionAxis = Action.Name.ToString();
+
+					// Parse comma delimited action names into an array
+					TArray<FString> ActionAxisArray;
+					ActionAxis.ParseIntoArray(ActionAxisArray, TEXT(","), true);
+					TArray<FInputAxisKeyMapping> FoundAxisMappings;
+
+					for (auto& ActionAxisName : ActionAxisArray)
+					{
+						FindAxisMappings(InputSettings, FName(*ActionAxisName), FoundAxisMappings);
+
+						for (FInputAxisKeyMapping AxisMapping : FoundAxisMappings)
+						{
+							if (UniqueInput.IsEqual(AxisMapping.Key.GetFName()))
+							{
+								// Check for X Axis
+								if (!Action.KeyX.IsNone() && Action.KeyX.IsEqual(AxisMapping.Key.GetFName()))
+								{
+									// Add 1D Action
+									NewAxisMapping.Actions.AddUnique(Action.Path);
+
+									FString ActionDimension = Action.Name.ToString().Right(7);
+
+									if (ActionDimension == TEXT("_axis2d"))
+									{
+										// Add 2D Action
+										FString Action2D = Action.Path.LeftChop(11) + TEXT(" X Y_axis2d");
+										NewAxisMapping.Actions.AddUnique(Action2D);
+									}
+
+									if (ActionDimension == TEXT("_axis3d"))
+									{
+										// Add 3D Action
+										FString Action3D = Action.Path.LeftChop(11) + TEXT(" X Y_axis3d");
+										NewAxisMapping.Actions.AddUnique(Action3D);
+									}
+								}
+							}
+						}
+					}
+				}
+
+				// Setup the action fields
+				TArray<FString> ActionFields = {
+												 TEXT("name"), Action.Path,
+												 TEXT("type"), Action.GetActionTypeName(),
+				};
+
+				// Add hand if skeleton
+				if (Action.Type == EActionType::Skeleton)
+				{
+					ActionFields.Append({ TEXT("skeleton"), Action.StringPath });
+				}
+
+				// Add optional field if this isn't a required field
+				if (!Action.bRequirement)
+				{
+					FString Optional[] = { TEXT("requirement"), TEXT("optional") };
+					ActionFields.Append(Optional, 2);
+				}
+
+				if (!UniqueActions.Contains(Action.Name.ToString()))
+				{
+					// Add this action to the array of input actions
+					TSharedRef<FJsonObject> ActionObject = MakeShareable(new FJsonObject());
+					BuildJsonObject(ActionFields, ActionObject);
+					InputActionsArray.AddUnique(MakeShareable(new FJsonValueObject(ActionObject)));
+
+					// Add this action to a cache of unique actions for this project
+					UniqueActions.AddUnique(Action.Name.ToString());
+
+					// Set localization text for this action
+					FString ActionName = Action.Name.ToString();
+					if (ActionName.Contains("axis"))
+					{
+						if (ActionName.Contains(","))
+						{
+							TArray<FString> ActionNameArray;
+							ActionName.ParseIntoArray(ActionNameArray, TEXT(","), true);
+
+							if (ActionNameArray.Num() > 0)
+							{
+								// Grab only the first action name & remove _X
+								ActionName = FString(ActionNameArray[0]).Replace(TEXT("_X"), TEXT("")); 
+							}
+						}
+						else
+						{
+							if (ActionName.Right(5).Equals(" axis"))
+							{
+								ActionName = ActionName.LeftChop(5); // Remove " axis" for the localization string
+							}	
+						}
+					}
+					FString LocalizationArray[] = { Action.Path, ActionName };
+					LocalizationFields.Append(LocalizationArray, 2);
+				}
+
+			}
+
+			// Add this Input Mapping to the main Input Mappings array
+			if (NewInputMapping.Actions.Num() > 0)
+			{
+				InputMappings.Add(NewInputMapping);
+			}
+
+			// Add this Axis Mapping to the main Input Mappings array
+			if (NewAxisMapping.Actions.Num() > 0)
+			{
+				InputMappings.Add(NewAxisMapping);
+			}
+		}
+
+		// If there are input actions, add them to the action manifest object
+		if (InputActionsArray.Num() > 0)
+		{
+			ActionManifestObject->SetArrayField(TEXT("actions"), InputActionsArray);
+		}
+	}
+	else
+	{
+		UE_LOG(LogSteamVRInputDevice, Error, TEXT("Error trying to retrieve Input Settings."));
+	}
+#pragma endregion
+
+#pragma region ACTION SETS
+	// Setup action set json objects
+	TArray<TSharedPtr<FJsonValue>> ActionSets;
+	TSharedRef<FJsonObject> ActionSetObject = MakeShareable(new FJsonObject());
+
+	// Create action set objects
+	TArray<FString> StringFields = {
+									 "name", TEXT(ACTION_SET),
+									 "usage", TEXT("leftright")
+	};
+
+	BuildJsonObject(StringFields, ActionSetObject);
+
+	// Add action sets array to the Action Manifest object
+	ActionSets.Add(MakeShareable(new FJsonValueObject(ActionSetObject)));
+	ActionManifestObject->SetArrayField(TEXT("action_sets"), ActionSets);
+
+	// Set localization text for the action set
+	LocalizationFields.Add(TEXT(ACTION_SET));
+	LocalizationFields.Add("Main Game Actions");
+#pragma endregion
+
+#pragma region DEFAULT CONTROLLER BINDINGST
+#if WITH_EDITOR
+	// If we're running in the editor, build the controller bindings if they don't exist yet
+	if (GenerateBindings)
+	{
+		GenerateControllerBindings(ControllerBindingsPath, ControllerTypes, ControllerBindings, Actions, InputMappings, DeleteIfExists);
+	}
+#endif
+
+	// Start search for controller bindings files
+	TArray<FString> ControllerBindingFiles;
+	FileManager.FindFiles(ControllerBindingFiles, *ControllerBindingsPath, TEXT("*.json"));
+	UE_LOG(LogSteamVRInputDevice, Log, TEXT("Searching for Controller Bindings files at: %s"), *ControllerBindingsPath);
+
+	// Look for existing controller binding files
+	for (FString& BindingFile : ControllerBindingFiles)
+	{
+		// Skip if manifest
+		if (BindingFile.Contains(TEXT("steamvr_manifest")))
+		{
+			continue;
+		}
+
+		// Setup cache
+		FString StringCache;
+		FString ControllerType;
+		
+		// Load Binding File to a string
+		FFileHelper::LoadFileToString(StringCache, *(ControllerBindingsPath / BindingFile));
+
+		// Convert string to json object
+		TSharedRef<TJsonReader<TCHAR>> JsonReader = TJsonReaderFactory<TCHAR>::Create(StringCache);
+		TSharedPtr<FJsonObject> JsonObject = MakeShareable(new FJsonObject());
+
+		// Attempt to deserialize string cache to a json object
+		if (!FJsonSerializer::Deserialize(JsonReader, JsonObject) || !JsonObject.IsValid())
+		{
+			UE_LOG(LogSteamVRInputDevice, Warning, TEXT("Invalid json format for controller binding file, skipping: %s"), *(ControllerBindingsPath / BindingFile));
+		}
+		// Attempt to find what controller this binding file is for 
+		else if (!JsonObject->TryGetStringField(TEXT("controller_type"), ControllerType) || ControllerType.IsEmpty())
+		{
+			UE_LOG(LogSteamVRInputDevice, Warning, TEXT("Unable to determine controller type for this binding file, skipping: %s"), *(ControllerBindingsPath / BindingFile));
+		}
+		else
+		{
+			// Create Controller Binding Object for this binding file
+			TSharedRef<FJsonObject> ControllerBindingObject = MakeShareable(new FJsonObject());
+			TArray<FString> ControllerStringFields = { "controller_type", *ControllerType,
+											 TEXT("binding_url"), *BindingFile //*FileManager.ConvertToAbsolutePathForExternalAppForRead(*(ControllerBindingsPath / BindingFile))
+			};
+			BuildJsonObject(ControllerStringFields, ControllerBindingObject);
+			ControllerBindings.Add(MakeShareable(new FJsonValueObject(ControllerBindingObject)));
+
+			// Tag this controller as generated
+			for (auto& DefaultControllerType : ControllerTypes)
+			{
+				if (DefaultControllerType.Name == FName(*ControllerType))
+				{
+					DefaultControllerType.bIsGenerated = true;
+				}
+			}
+		}
+	}
+
+	// Add the default bindings object to the action manifest
+	if (ControllerBindings.Num() == 0)
+	{
+		UE_LOG(LogSteamVRInputDevice, Error, TEXT("Unable to find and/or generate controller binding files in: %s"), *ControllerBindingsPath);
+	}
+	else
+	{
+		ActionManifestObject->SetArrayField(TEXT("default_bindings"), ControllerBindings);
+	}
+#pragma endregion
+
+#pragma region LOCALIZATION
+	// Setup localizations json objects
+	TArray<TSharedPtr<FJsonValue>> Localizations;
+	TSharedRef<FJsonObject> LocalizationsObject = MakeShareable(new FJsonObject());
+
+	// Build & add localizations to the Action Manifest object
+	BuildJsonObject(LocalizationFields, LocalizationsObject);
+	Localizations.Add(MakeShareable(new FJsonValueObject(LocalizationsObject)));
+	ActionManifestObject->SetArrayField(TEXT("localization"), Localizations);
+#pragma endregion
+
+	// Serialize Action Manifest Object
+	FString ActionManifest;
+	TSharedRef<TJsonWriter<>> JsonWriter = TJsonWriterFactory<>::Create(&ActionManifest);
+	FJsonSerializer::Serialize(ActionManifestObject, JsonWriter);
+
+	// Save json as a UTF8 file
+	if (GenerateActions)
+	{
+		if (FileManager.FileExists(*ManifestPath) && DeleteIfExists)
+		{
+			FPlatformFileManager::Get().GetPlatformFile().DeleteFile(*ManifestPath);
+		}
+
+		if (!FileManager.FileExists(*ManifestPath))
+		{
+			if (!FFileHelper::SaveStringToFile(ActionManifest, *ManifestPath, FFileHelper::EEncodingOptions::ForceUTF8WithoutBOM))
+			{
+				UE_LOG(LogSteamVRInputDevice, Error, TEXT("Error trying to generate action manifest in: %s"), *ManifestPath);
+				return;
+			}
+		}
+	}
+
+	// Register Application to SteamVR
+	if (RegisterApp)
+	{
+		RegisterApplication(ManifestPath);
+	}
+
+	// Update action events
+	ActionEvents.Empty();
+	bool bAlreadyExists = false;
+	for (FSteamVRInputAction& InputAction : Actions)
+	{
+		// Check if we've already got a similar action event as we won't need the flat action list for processing controller events
+		bAlreadyExists = false;
+		for (FSteamVRInputAction& ActionEvent : ActionEvents)
+		{
+			if (ActionEvent.Handle == InputAction.Handle)
+			{
+				bAlreadyExists = true;
+				break;
+			}
+		}
+
+		if (!bAlreadyExists && InputAction.Handle != k_ulInvalidActionHandle)
+		{
+			ActionEvents.Add(FSteamVRInputAction(InputAction));
+		}
+	}
+}
+
+bool FSteamVRInputDevice::BuildJsonObject(TArray<FString> StringFields, TSharedRef<FJsonObject> OutJsonObject)
+{
+	// Check if StringFields array is even
+	if (StringFields.Num() > 1 && StringFields.Num() % 2 == 0)
+	{
+		// Generate json object of string field pairs
+		for (int32 i = 0; i < StringFields.Num(); i += 2)
+		{
+			OutJsonObject->SetStringField(StringFields[i], StringFields[i + 1]);
+		}
+
+		return true;
+	}
+
+	return false;
+}
+
+void FSteamVRInputDevice::ProcessKeyInputMappings(const UInputSettings* InputSettings, TArray<FName> &InOutUniqueInputs)
+{
+	// Retrieve key actions setup in this project
+	KeyMappings.Empty();
+	SteamVRKeyInputMappings.Empty();
+	TArray<FName> KeyActionNames;
+	InputSettings->GetActionNames(KeyActionNames);
+
+	// Process all key actions found
+	for (const FName& KeyActionName : KeyActionNames)
+	{
+		TArray<FInputActionKeyMapping> KeyInputMappings;
+
+		// Retrieve input keys associated with this action
+		FindActionMappings(InputSettings, KeyActionName, KeyInputMappings);
+
+		for (auto& KeyMapping : KeyInputMappings)
+		{
+			//UE_LOG(LogSteamVRInputDevice, Warning, TEXT("*** Processing Action: %s"), *KeyMapping.ActionName.ToString());
+
+			// Default to "MotionController" Generic UE type
+			FString CurrentControllerType = FString(TEXT("MotionController"));
+
+			// Get the string version of the key id we are dealing with for analysis
+			FString CurrentKey = KeyMapping.Key.GetFName().ToString();
+
+			// Determine which supported controller type we are working with
+			if (CurrentKey.Contains(TEXT("ValveIndex")))
+			{
+				CurrentControllerType = FString(TEXT("ValveIndex"));
+			}
+			else if (CurrentKey.Contains(TEXT("Vive")))
+			{
+				CurrentControllerType = FString(TEXT("Vive"));
+			}
+			else if (CurrentKey.Contains(TEXT("Cosmos")))
+			{
+				CurrentControllerType = FString(TEXT("Cosmos"));
+			}
+			else if (CurrentKey.Contains(TEXT("OculusTouch")))
+			{
+				CurrentControllerType = FString(TEXT("OculusTouch"));
+			}
+			else if (CurrentKey.Contains(TEXT("MixedReality")))
+			{
+				CurrentControllerType = FString(TEXT("MixedReality"));
+			}
+			else if (CurrentKey.Contains(TEXT("HMD_Proximity")))
+			{
+				CurrentControllerType = FString(TEXT("HMD_Proximity"));
+			}
+			else if (CurrentKey.Contains(TEXT("MotionController")))
+			{
+				// retain default
+			}
+			else
+			{
+				continue; // unrecognized controller - will not process
+			}
+
+			// Only process Motion Controller if there are no SteamVR actions
+			if (KeyMapping.Key.GetFName().ToString().Contains(TEXT("MotionController")))
+			{
+				bool bFound = false;
+				for (FInputActionKeyMapping& KeyMappingInner : KeyInputMappings)
+				{
+					//UE_LOG(LogTemp, Warning, TEXT("SEARCH: \nKeyMappingKey: %s \nKeyInner:%s \n%s \n%s"),
+					//	*KeyMapping.Key.GetFName().ToString(),
+					//	*KeyMappingInner.Key.GetFName().ToString(),
+					//	*KeyMappingInner.Key.GetFName().ToString(),
+					//	*KeyMappingInner.Key.GetFName().ToString());
+
+					if (KeyMapping.Key.GetFName().ToString().Contains(KeyMappingInner.Key.GetFName().ToString())
+						&& IsVRKey(KeyMappingInner.Key.GetFName())
+						)
+					{
+						//UE_LOG(LogTemp, Warning, TEXT("SEARCH: MOTION CONTROLLER with STEAMVR Paired action found!"));
+						bFound = true;
+						break;
+					}
+				}
+
+				if (bFound)
+				{
+					continue;
+				}
+			}
+
+			// If there's a Motion Controller or valid device input, add to the SteamVR Input Actions
+			Actions.Add(FSteamVRInputAction(
+				FString(ACTION_PATH_IN) / KeyActionName.ToString(),
+				KeyActionName,
+				KeyMapping.Key.GetFName(),
+				false));
+
+			// Add input names here for use in the auto-generation of controller bindings
+			InOutUniqueInputs.AddUnique(KeyMapping.Key.GetFName());
+
+			// Add input to Key Bindings Cache
+			FSteamVRInputKeyMapping SteamVRInputKeyMap = FSteamVRInputKeyMapping(KeyMapping);
+			SteamVRInputKeyMap.ActionName = KeyActionName.ToString();
+			SteamVRInputKeyMap.ActionNameWithPath = FString(ACTION_PATH_IN) / KeyActionName.ToString();
+			SteamVRInputKeyMap.ControllerName = CurrentControllerType;				
+			SteamVRKeyInputMappings.Add(SteamVRInputKeyMap);
+		}
+	}
+}
+
+void FSteamVRInputDevice::ProcessKeyAxisMappings(const UInputSettings* InputSettings, TArray<FName> &InOutUniqueInputs)
+{
+	// Retrieve Key Axis names
+	TArray<FName> KeyAxisNames;
+	InputSettings->GetAxisNames(KeyAxisNames);
+	KeyAxisMappings.Empty();
+	SteamVRKeyAxisMappings.Empty();
+
+	// Iterate over every axis name found in this project, and process for Vector 1, 2 or 3
+	for (const FName& XAxisName : KeyAxisNames)
+	{
+		// Set X Axis Key Name Cache
+		FName XAxisNameKey = NAME_None;	
+		FName YAxisNameKey = NAME_None;
+		FName YAxisName = NAME_None;
+		FName ZAxisNameKey = NAME_None;
+		FName ZAxisName = NAME_None;
+
+		// Retrieve input axes associated with this action
+		FindAxisMappings(InputSettings, XAxisName, KeyAxisMappings);
+
+		// Create a SteamVR Axis Key Mapping that holds metadata for us
+		GetSteamVRMappings(KeyAxisMappings, SteamVRKeyAxisMappings);
+
+		// STEP 1: Go through all X axis mappings, checking for which type of Vector this is (1, 2 or 3)
+		for (FSteamVRAxisKeyMapping& AxisMapping : SteamVRKeyAxisMappings)
+		{
+			// Add axes names here for use in the auto-generation of controller bindings
+			InOutUniqueInputs.AddUnique(AxisMapping.InputAxisKeyMapping.Key.GetFName());
+
+			// Default to "MotionController" Generic UE type
+			FString CurrentControllerType = FString(TEXT("MotionController"));
+
+			// If this is an X Axis key, check for the corresponding Y & Z Axes as well
+			uint32 KeyHand = 0;
+
+			// Get the string version of the key id we are dealing with for analysis
+			FString CurrentKey = AxisMapping.InputAxisKeyMapping.Key.GetFName().ToString();
+			
+			// Determine which supported controller type we are working with
+			if (CurrentKey.Contains(TEXT("ValveIndex")))
+			{
+				CurrentControllerType = FString(TEXT("ValveIndex"));
+			}
+			else if (CurrentKey.Contains(TEXT("Vive")))
+			{
+				CurrentControllerType = FString(TEXT("Vive"));
+			}
+			else if (CurrentKey.Contains(TEXT("Cosmos")))
+			{
+				CurrentControllerType = FString(TEXT("Cosmos"));
+			}
+			else if (CurrentKey.Contains(TEXT("OculusTouch")))
+			{
+				CurrentControllerType = FString(TEXT("OculusTouch"));
+			}
+			else if (CurrentKey.Contains(TEXT("MixedReality")))
+			{
+				CurrentControllerType = FString(TEXT("MixedReality"));
+			}
+			else if (CurrentKey.Contains(TEXT("MotionController")))
+			{
+				// empty on purpose (readability)
+			}
+			else
+			{
+				continue;
+			}
+
+			// Set the Controller Type for this axis mapping
+			AxisMapping.ControllerName = CurrentControllerType;
+
+			// Create a Y Equivalent of the X Action to ensure we are matching the action and not just the controller type
+			FString CurrentActionName_Y = AxisMapping.InputAxisKeyMapping.AxisName.ToString().Replace(TEXT("_X"), TEXT("_Y"));
+			FString CurrentActionName_Z = AxisMapping.InputAxisKeyMapping.AxisName.ToString().Replace(TEXT("_X"), TEXT("_Z"));
+
+			// Convert the controller key id name to a string we can do some quick checks on it
+			FString KeyString_X = AxisMapping.InputAxisKeyMapping.Key.GetFName().ToString();
+			FString KeyString_Y = "";
+			FString KeyString_Z = "";
+			bool bIsXAxis = false;
+
+			if (KeyString_X.Contains(TEXT("_X_"), ESearchCase::CaseSensitive, ESearchDir::FromEnd))
+			{
+				bIsXAxis = true;
+				KeyString_Y = KeyString_X.Replace(TEXT("_X_"), TEXT("_Y_"));
+				KeyString_Z = KeyString_X.Replace(TEXT("_X_"), TEXT("_Z_"));
+			}
+			else if (KeyString_X.Contains(TEXT("_X"), ESearchCase::CaseSensitive, ESearchDir::FromEnd))
+			{
+				bIsXAxis = true;
+				KeyString_Y = KeyString_X.Replace(TEXT("_X"), TEXT("_Y"));
+				KeyString_Z = KeyString_X.Replace(TEXT("_X"), TEXT("_Z"));
+			}
+			else if (KeyString_X.Contains(TEXT("X-Axis"), ESearchCase::CaseSensitive, ESearchDir::FromEnd))
+			{
+				bIsXAxis = true;
+				KeyString_Y = KeyString_X.Replace(TEXT("X-Axis"), TEXT("Y-Axis"));
+				KeyString_Z = KeyString_X.Replace(TEXT("X-Axis"), TEXT("Z-Axis"));
+			}
+
+
+			// Check if this controller is meant to be a float X axis key
+			if (bIsXAxis)
+			{
+				// Set X Axis
+				XAxisNameKey = FName(*KeyString_X);
+
+				// Go through all the axis names again looking for Y and Z inputs that correspond to this X input
+				for (const FName& KeyAxisNameInner : KeyAxisNames)
+				{
+					// Retrieve input axes associated with this action
+					TArray<FInputAxisKeyMapping> AxisMappingsInner;
+					FindAxisMappings(InputSettings, KeyAxisNameInner, AxisMappingsInner);
+
+					for (auto& AxisMappingInner : AxisMappingsInner)
+					{
+						// Get the name of the key for this action
+						FString KeyNameString = (AxisMappingInner.Key.GetFName().ToString());
+
+						// Check if we are dealing with the same controller for this action
+						if (!KeyNameString.Contains(CurrentControllerType))
+						{
+							// Skip this action if not
+							continue;
+						}
+
+						// Check if this is an equivalent Y Axis key for our current X Axis key
+						if (KeyString_Y.Equals(KeyNameString) && AxisMappingInner.AxisName.ToString().Equals(CurrentActionName_Y))
+						{
+							YAxisName = FName(KeyAxisNameInner);
+							YAxisNameKey = FName(*KeyString_Y);
+							AxisMapping.bIsPartofVector2 = true;
+						}
+						else if (KeyString_Z.Equals(KeyNameString) && AxisMappingInner.AxisName.ToString().Equals(CurrentActionName_Z))
+						{
+							ZAxisName = KeyAxisNameInner;
+							ZAxisNameKey = AxisMappingInner.Key.GetFName();
+							AxisMapping.bIsPartofVector3 = true;
+						}
+					}
+				}
+
+				// Set the Axis Names
+				if (YAxisName != NAME_None && ZAxisName == NAME_None)
+				{
+					// [2D] There's a Y Axis but no Z, this must be a Vector2
+					AxisMapping.XAxisName = FName(AxisMapping.InputAxisKeyMapping.AxisName);
+					AxisMapping.YAxisName = FName(YAxisName);
+					
+					AxisMapping.XAxisKey = FName(XAxisNameKey);
+					AxisMapping.YAxisKey = FName(YAxisNameKey);
+					
+					AxisMapping.bIsPartofVector2 = true;
+				}
+				else if (YAxisName != NAME_None && ZAxisName != NAME_None)
+				{
+					// [3D] There's a Z Axis, this must be a Vector3
+					AxisMapping.XAxisName = FName(AxisMapping.InputAxisKeyMapping.AxisName);
+					AxisMapping.YAxisName = FName(YAxisName);
+					AxisMapping.ZAxisName = FName(ZAxisName);
+					
+					AxisMapping.XAxisKey = FName(XAxisNameKey);
+					AxisMapping.YAxisKey = FName(YAxisNameKey);
+					AxisMapping.ZAxisKey = FName(ZAxisNameKey);
+
+					AxisMapping.bIsPartofVector3 = true;
+				}
+
+				// Reset Name Caches
+				YAxisNameKey = NAME_None;
+				YAxisName = NAME_None;
+				ZAxisNameKey = NAME_None;
+				ZAxisName = NAME_None;
+			}
+		}
+
+		// STEP 2: Go through all Y axis mappings, checking for which type of Vector this is (1, 2 or 3)
+		for (auto& AxisMapping : SteamVRKeyAxisMappings)
+		{
+			// Add axes names here for use in the auto-generation of controller bindings
+			InOutUniqueInputs.AddUnique(AxisMapping.InputAxisKeyMapping.Key.GetFName());
+
+			// Default to "MotionController" Generic UE type
+			FString CurrentControllerType = FString(TEXT("MotionController"));
+
+			// If this is an X Axis key, check for the corresponding Y & Z Axes as well
+			uint32 KeyHand = 0;
+
+			// Get the string version of the key id we are dealing with for analysis
+			FString CurrentKey = AxisMapping.InputAxisKeyMapping.Key.GetFName().ToString();
+
+			// Determine which supported controller type we are working with
+			if (CurrentKey.Contains(TEXT("ValveIndex")))
+			{
+				CurrentControllerType = FString(TEXT("ValveIndex"));
+			}
+			else if (CurrentKey.Contains(TEXT("Vive")))
+			{
+				CurrentControllerType = FString(TEXT("Vive"));
+			}
+			else if (CurrentKey.Contains(TEXT("Cosmos")))
+			{
+				CurrentControllerType = FString(TEXT("Cosmos"));
+			}
+			else if (CurrentKey.Contains(TEXT("OculusTouch")))
+			{
+				CurrentControllerType = FString(TEXT("OculusTouch"));
+			}
+			else if (CurrentKey.Contains(TEXT("MixedReality")))
+			{
+				CurrentControllerType = FString(TEXT("MixedReality"));
+			} 
+			else if (CurrentKey.Contains(TEXT("MotionController")))
+			{
+				// empty on purpose (readability)
+			}
+			else
+			{
+				continue;
+			}
+
+			// Set the Controller Type for this axis mapping
+			AxisMapping.ControllerName = CurrentControllerType;
+
+			// Convert the controller key id name to a string so we can do some quick checks on it
+			FString KeyString_X = "";
+			FString KeyString_Y = AxisMapping.InputAxisKeyMapping.Key.GetFName().ToString();
+			FString KeyString_Z = "";
+			bool bIsYAxis = false;
+
+			if (KeyString_Y.Contains(TEXT("_Y_"), ESearchCase::CaseSensitive, ESearchDir::FromEnd))
+			{
+				bIsYAxis = true;
+				KeyString_X = KeyString_Y.Replace(TEXT("_Y_"), TEXT("_X_"));
+				KeyString_Z = KeyString_Y.Replace(TEXT("_Y_"), TEXT("_Z_"));
+			}
+			else if (KeyString_Y.Contains(TEXT("_Y"), ESearchCase::CaseSensitive, ESearchDir::FromEnd))
+			{
+				bIsYAxis = true;
+				KeyString_X = KeyString_Y.Replace(TEXT("_Y"), TEXT("_X"));
+				KeyString_Z = KeyString_Y.Replace(TEXT("_Y"), TEXT("_Z"));
+			}
+			else if (KeyString_Y.Contains(TEXT("Y-Axis"), ESearchCase::CaseSensitive, ESearchDir::FromEnd))
+			{
+				bIsYAxis = true;
+				KeyString_X = KeyString_Y.Replace(TEXT("Y-Axis"), TEXT("X-Axis"));
+				KeyString_Z = KeyString_Y.Replace(TEXT("Y-Axis"), TEXT("Z-Axis"));
+			}
+
+			// Check if this controller is meant to be a float Y axis key
+			if (bIsYAxis)
+			{
+				// Go through all the axis names again looking for X and Z inputs that correspond to this Y input
+				for (const FName& KeyAxisNameInner : KeyAxisNames)
+				{
+					// Retrieve input axes associated with this action
+					TArray<FInputAxisKeyMapping> AxisMappingsInner;
+					FindAxisMappings(InputSettings, KeyAxisNameInner, AxisMappingsInner);
+
+					for (auto& AxisMappingInner : AxisMappingsInner)
+					{
+						// Get the name of the key for this action
+						FString KeyNameString = (AxisMappingInner.Key.GetFName().ToString());
+
+						// Check if we are dealing with the same controller for this action
+						if (!KeyNameString.Contains(CurrentControllerType))
+						{
+							// Skip this action if not
+							continue;
+						}
+
+						// Check if this is an equivalent X Axis key for our current Y Axis key
+						if (KeyString_X.Equals(KeyNameString))
+						{
+							AxisMapping.bIsPartofVector2 = true;
+						}
+						else if (KeyString_Z.Equals(KeyNameString))
+						{
+							AxisMapping.bIsPartofVector3 = true;
+						}
+					}
+				}
+			}
+		}
+	}
+
+	// STEP 3: Create the axis action names
+	for (FSteamVRAxisKeyMapping& AxisMapping : SteamVRKeyAxisMappings)
+	{
+		// Only process valid controllers
+		if ((!IsVRKey(AxisMapping.InputAxisKeyMapping.Key.GetFName())
+			&& !AxisMapping.InputAxisKeyMapping.Key.GetFName().ToString().Contains(TEXT("MotionController"))))
+		{
+			continue;
+		} 
+
+		// Only process Motion Controller if there are no SteamVR actions
+		if (AxisMapping.InputAxisKeyMapping.Key.GetFName().ToString().Contains(TEXT("MotionController")))
+		{
+			bool bFound = false;
+			for (FSteamVRAxisKeyMapping& AxisMappingInner : SteamVRKeyAxisMappings)
+			{
+				//UE_LOG(LogTemp, Warning, TEXT("SEARCH: \nActionName: %s \nActionNameInner:%s \n%s \n%s"), 
+				//	*AxisMapping.InputAxisKeyMapping.AxisName.ToString(), 
+				//	*AxisMappingInner.InputAxisKeyMapping.AxisName.ToString(),
+				//	*AxisMapping.InputAxisKeyMapping.Key.GetFName().ToString(),
+				//	*AxisMappingInner.InputAxisKeyMapping.Key.GetFName().ToString());
+
+				if (AxisMapping.InputAxisKeyMapping.AxisName.ToString().Contains(AxisMappingInner.InputAxisKeyMapping.AxisName.ToString())
+					&& IsVRKey(AxisMappingInner.InputAxisKeyMapping.Key.GetFName()))
+				{
+					//UE_LOG(LogTemp, Warning, TEXT("SEARCH: MOTION CONTROLLER with STEAMVR Paired action found!"));
+					bFound = true;
+					break;
+				}
+			}
+
+			if (bFound)
+			{
+				continue;
+			}
+		}
+
+		if (AxisMapping.bIsPartofVector2)
+		{
+			// Check for empty actions
+			if (AxisMapping.InputAxisKeyMapping.AxisName == NAME_None ||
+				AxisMapping.YAxisName == NAME_None
+				)
+			{
+				AxisMapping.bIsPartofVector2 = false;
+			}
+			else
+			{
+				// Add a Vector 2 Action to our Actions list
+				FString AxisName2D = AxisMapping.InputAxisKeyMapping.AxisName.ToString() +
+					TEXT(",") +
+					AxisMapping.YAxisName.ToString() +
+					TEXT(" X Y_axis2d");
+				FString ActionPath2D = FString(ACTION_PATH_IN) / AxisName2D;
+
+				Actions.Add(FSteamVRInputAction(ActionPath2D, FName(*AxisName2D), AxisMapping.XAxisKey, AxisMapping.YAxisKey, FVector2D()));
+				AxisMapping.ActionName = FString(AxisName2D);
+				AxisMapping.ActionNameWithPath = FString(ActionPath2D);
+			}
+		}
+		else if (AxisMapping.bIsPartofVector3)
+		{
+			// Check for empty actions
+			if (AxisMapping.InputAxisKeyMapping.AxisName == NAME_None ||
+				AxisMapping.YAxisName == NAME_None ||
+				AxisMapping.ZAxisName == NAME_None
+				)
+			{
+				AxisMapping.bIsPartofVector3 = false;
+			}
+		}
+		else
+		{
+			// Add a Vector 1 to our Actions List
+			FString AxisName1D = AxisMapping.InputAxisKeyMapping.AxisName.ToString() + TEXT(" axis");
+			FString ActionPath = FString(ACTION_PATH_IN) / AxisName1D;
+			Actions.Add(FSteamVRInputAction(ActionPath, FName(*AxisName1D), AxisMapping.InputAxisKeyMapping.Key.GetFName(), 0.0f));
+			AxisMapping.ActionName = FString(AxisName1D);
+			AxisMapping.ActionNameWithPath = FString(ActionPath);
+		}
+	}
+
+	// Cleanup action set
+	SanitizeActions();
+}
+
+void FSteamVRInputDevice::SanitizeActions()
+{
+	UInputSettings* InputSettings = GetMutableDefault<UInputSettings>();
+	TArray<FInputAxisKeyMapping> DuplicateActions;
+
+	if (InputSettings->IsValidLowLevelFast())
+	{
+		// Check for duplicates (can be safely removed once full support for Vector2 is up)
+		for (auto AxisKeyMapping : InputSettings->GetAxisMappings())
+		{
+			bool bIsDuplicate = false;
+			for (auto DuplicateAction : DuplicateActions)
+			{
+				if (DuplicateAction.Key == AxisKeyMapping.Key)
+				{
+					bIsDuplicate = true;
+					break;
+				}
+			}
+		}
+	}
+
+	// Remove duplicates
+	for (auto AxisKeyMapping : DuplicateActions)
+	{
+		InputSettings->RemoveAxisMapping(AxisKeyMapping);
+	}
+
+	// Save to config files and cleanup
+	if (DuplicateActions.Num() > 0)
+	{
+		InputSettings->SaveKeyMappings();
+		InputSettings->UpdateDefaultConfigFile();
+		DuplicateActions.Empty();
+	}
+}
+
+void FSteamVRInputDevice::RegisterApplication(FString ManifestPath)
+{
+
+	if (VRSystem() && VRInput())
+	{
+		// Get Project Name this plugin is used in
+		uint32 AppProcessId = FPlatformProcess::GetCurrentProcessId();
+		GameFileName = FPaths::GetCleanFilename(FPlatformProcess::GetApplicationName(AppProcessId));
+		GameProjectName = GameFileName;
+		if (GConfig)
+		{
+			GConfig->GetString(
+				TEXT("/Script/EngineSettings.GeneralProjectSettings"),
+				TEXT("ProjectName"),
+				GameProjectName,
+				GGameIni
+			);
+
+		}
+
+		// Check for empty project name
+		if (GameProjectName.IsEmpty())
+		{
+			GameProjectName = GameFileName;
+		}
+
+#if WITH_EDITOR
+		if (VRApplications())
+		{
+			// Generate Application Manifest
+			FString AppKey, AppManifestPath;
+	
+			GenerateAppManifest(ManifestPath, GameFileName, AppKey, AppManifestPath);
+	
+			char* SteamVRAppKey = TCHAR_TO_UTF8(*AppKey);
+	
+			// Load application manifest
+			EVRApplicationError AppError = VRApplications()->AddApplicationManifest(TCHAR_TO_UTF8(*IFileManager::Get().ConvertToAbsolutePathForExternalAppForRead(*AppManifestPath)), true);
+			UE_LOG(LogSteamVRInputDevice, Display, TEXT("[STEAMVR INPUT] Registering Application Manifest %s : %s"), *AppManifestPath, *FString(UTF8_TO_TCHAR(VRApplications()->GetApplicationsErrorNameFromEnum(AppError))));
+	
+			// Set AppKey for this Editor Session
+			AppError = VRApplications()->IdentifyApplication(AppProcessId, SteamVRAppKey);
+			UE_LOG(LogSteamVRInputDevice, Display, TEXT("[STEAMVR INPUT] Editor Application [%d][%s] identified to SteamVR: %s"), AppProcessId, *AppKey, *FString(UTF8_TO_TCHAR(VRApplications()->GetApplicationsErrorNameFromEnum(AppError))));
+		}
+#endif
+
+		// Set Action Manifest
+		FString TheActionManifestPath;
+		
+		#if WITH_EDITOR
+			TheActionManifestPath = IFileManager::Get().ConvertToAbsolutePathForExternalAppForRead(*ManifestPath);
+		#else
+			TheActionManifestPath = FPaths::ConvertRelativePathToFull(FPaths::ProjectDir() / TEXT("Config") / TEXT("SteamVRBindings") / TEXT(ACTION_MANIFEST)).Replace(TEXT("/"), TEXT("\\"));
+		#endif
+		
+		UE_LOG(LogSteamVRInputDevice, Display, TEXT("[STEAMVR INPUT] Trying to load Action Manifest from: %s"), *TheActionManifestPath);
+		EVRInputError InputError = VRInput()->SetActionManifestPath(TCHAR_TO_UTF8(*TheActionManifestPath));
+		GetInputError(InputError, FString(TEXT("Setting Action Manifest Path Result")));
+
+		// Set Main Action Set
+		InputError = VRInput()->GetActionSetHandle(ACTION_SET, &MainActionSet);
+		GetInputError(InputError, FString(TEXT("Setting main action set")));
+
+		// Add to action set array
+		SteamVRInputActionSets.Add(FSteamVRInputActionSet(0, ACTION_SET, MainActionSet));
+
+		// Populate Active Action sets that will later be used in OpenVR calls
+		for (int32 i = 0; i < SteamVRInputActionSets.Num(); i++)
+		{
+			if (i > MAX_ACTION_SETS - 1) break;	// Skip if more than allocated maximum action sets
+
+			ActiveActionSets[i].nPriority = (int32_t)SteamVRInputActionSets[i].Priority;
+			ActiveActionSets[i].ulActionSet = SteamVRInputActionSets[i].Handle;
+			ActiveActionSets[i].ulRestrictedToDevice = SteamVRInputActionSets[i].RestrictedToDeviceHandle;
+			ActiveActionSets[i].ulSecondaryActionSet = SteamVRInputActionSets[i].SecondaryActionSetHandle;
+		}
+
+		// Fill in Action handles for each registered action
+		for (auto& Action : Actions)
+		{
+			VRActionHandle_t Handle;
+			InputError = VRInput()->GetActionHandle(TCHAR_TO_UTF8(*Action.Path), &Handle);
+
+			if (InputError != VRInputError_None || Handle == k_ulInvalidActionHandle)
+			{
+				continue;
+			}
+
+			Action.Handle = Handle;
+
+			// Test if this is a pose
+			if (Action.Path == TEXT(ACTION_PATH_CONTROLLER_LEFT))
+			{
+				VRControllerHandleLeft = Action.Handle;
+			}
+			else if (Action.Path == TEXT(ACTION_PATH_CONTROLLER_RIGHT))
+			{
+				VRControllerHandleRight = Action.Handle;
+			}
+			else if (Action.Path == TEXT(ACTION_PATH_SPECIAL_BACK_L))
+			{
+				VRSpecial1 = Action.Handle;
+			}
+			else if (Action.Path == TEXT(ACTION_PATH_SPECIAL_BACK_R))
+			{
+				VRSpecial2 = Action.Handle;
+			}
+			else if (Action.Path == TEXT(ACTION_PATH_SPECIAL_FRONT_L))
+			{
+				VRSpecial3 = Action.Handle;
+			}
+			else if (Action.Path == TEXT(ACTION_PATH_SPECIAL_FRONT_R))
+			{
+				VRSpecial4 = Action.Handle;
+			}
+			else if (Action.Path == TEXT(ACTION_PATH_SPECIAL_FRONTR_L))
+			{
+				VRSpecial5 = Action.Handle;
+			}
+			else if (Action.Path == TEXT(ACTION_PATH_SPECIAL_FRONTR_R))
+			{
+				VRSpecial6 = Action.Handle;
+			}
+			else if (Action.Path == TEXT(ACTION_PATH_SPECIAL_PISTOL_L))
+			{
+				VRSpecial7 = Action.Handle;
+			}
+			else if (Action.Path == TEXT(ACTION_PATH_SPECIAL_PISTOL_R))
+			{
+				VRSpecial8 = Action.Handle;
+			}
+
+			UE_LOG(LogSteamVRInputDevice, Display, TEXT("Retrieving Action Handle: %s"), *Action.Path);
+			GetInputError(InputError, FString(TEXT("Setting Action Handle Path Result")));
+		}
+	}
+}
+
+bool FSteamVRInputDevice::SetSkeletalHandle(char* ActionPath, VRActionHandle_t& SkeletalHandle)
+{
+	if (VRSystem() && VRInput())
+	{
+		// Get Skeletal Handle
+		EVRInputError Err = VRInput()->GetActionHandle(ActionPath, &SkeletalHandle);
+		if (Err != VRInputError_None || SkeletalHandle == k_ulInvalidActionHandle)
+		{
+			if (Err != LastInputError)
+			{
+				GetInputError(Err, TEXT("Couldn't get skeletal action handle for Skeleton."));
+			}
+
+			Err = LastInputError;
+			return false;
+		}
+		else
+		{
+			Err = LastInputError;
+			return true;
+		}
+	}
+	return false;
+}
+
+void FSteamVRInputDevice::ProcessActionEvents(FSteamVRInputActionSet SteamVRInputActionSet)
+{
+	for (auto& Action : ActionEvents)
+	{
+		if (Action.Handle == k_ulInvalidActionHandle)
+		{
+			continue;
+		}
+
+		if (Action.Type == EActionType::Boolean && !Action.Path.Contains(TEXT(" axis")) && !Action.Path.Contains(TEXT("_axis")))
+		{
+			// Get digital data from SteamVR
+			InputDigitalActionData_t DigitalData;
+			EVRInputError ActionStateError = VRInput()->GetDigitalActionData(Action.Handle, &DigitalData, sizeof(DigitalData), k_ulInvalidInputValueHandle);
+
+			if (ActionStateError != VRInputError_None)
+			{
+				//GetInputError(ActionStateError, TEXT("Error encountered retrieving digital data from SteamVR"));
+				//UE_LOG(LogSteamVRInputDevice, Error, TEXT("%s"), *Action.Path);
+				continue;
+			}
+			else if (ActionStateError == VRInputError_None &&
+				DigitalData.bActive &&
+				DigitalData.bState != Action.bState
+				)
+			{
+				// Update our action to the value reported by SteamVR
+				Action.bState = DigitalData.bState;
+
+				// Update the active origin
+				Action.ActiveOrigin = DigitalData.activeOrigin;
+
+				// Test what we're receiving from SteamVR
+				//UE_LOG(LogTemp, Warning, TEXT("Handle %s KeyX %s Value %i"), *Action.Path, *FoundKey.GetFName().ToString(), DigitalData.bState);
+
+				// Sent event back to Engine
+				if (Action.bState && Action.KeyX != NAME_None)
+				{
+					MessageHandler->OnControllerButtonPressed(Action.KeyX, 0, false);
+				}
+				else
+				{
+					MessageHandler->OnControllerButtonReleased(Action.KeyX, 0, false);
+				}
+			}
+		}
+		else if (Action.Type == EActionType::Vector1
+			|| Action.Type == EActionType::Vector2
+			|| Action.Type == EActionType::Vector3)
+		{
+			// Get analog data from SteamVR
+			InputAnalogActionData_t AnalogData;
+			EVRInputError ActionStateError = VRInput()->GetAnalogActionData(Action.Handle, &AnalogData, sizeof(AnalogData), k_ulInvalidInputValueHandle);
+
+			if (ActionStateError != VRInputError_None)
+			{
+				//GetInputError(ActionStateError, TEXT("Error encountered retrieving analog data from SteamVR"));
+				continue;
+			}
+			else if (ActionStateError == VRInputError_None && AnalogData.bActive)
+			{
+				// Update the active origin
+				Action.ActiveOrigin = AnalogData.activeOrigin;
+
+				if (Action.KeyX != NAME_None)
+				{
+					// Test what we're receiving from SteamVR
+					//UE_LOG(LogTemp, Warning, TEXT("Handle %s KeyX %s X-Value [%f]"), *Action.Path, *FoundKeyX.GetFName().ToString(), AnalogData.x);
+
+					Action.Value.X = AnalogData.x; // ActionCount;
+					MessageHandler->OnControllerAnalog(Action.KeyX, 0, Action.Value.X);
+				}
+
+				if (Action.KeyY != NAME_None)
+				{
+					// Test what we're receiving from SteamVR
+					//UE_LOG(LogTemp, Warning, TEXT("Handle %s KeyY %s Y-Value {%f}"), *Action.Path, *FoundKeyY.GetFName().ToString(), AnalogData.y);
+
+					Action.Value.Y = AnalogData.y;
+					MessageHandler->OnControllerAnalog(Action.KeyY, 0, Action.Value.Y);
+				}
+			}
+		}
+	}
+}
+
+void FSteamVRInputDevice::GetInputError(EVRInputError InputError, FString InputAction)
+{
+	switch (InputError)
+	{
+	case VRInputError_None:
+		UE_LOG(LogSteamVRInputDevice, Display, TEXT("[STEAMVR INPUT] %s: Success"), *InputAction);
+		break;
+	case VRInputError_NameNotFound:
+		UE_LOG(LogSteamVRInputDevice, Error, TEXT("[STEAMVR INPUT] %s: Name Not Found"), *InputAction);
+		break;
+	case VRInputError_WrongType:
+		UE_LOG(LogSteamVRInputDevice, Error, TEXT("[STEAMVR INPUT] %s: Wrong Type"), *InputAction);
+		break;
+	case VRInputError_InvalidHandle:
+		UE_LOG(LogSteamVRInputDevice, Error, TEXT("[STEAMVR INPUT] %s: Invalid Handle"), *InputAction);
+		break;
+	case VRInputError_InvalidParam:
+		UE_LOG(LogSteamVRInputDevice, Error, TEXT("[STEAMVR INPUT] %s: Invalid Param"), *InputAction);
+		break;
+	case VRInputError_NoSteam:
+		UE_LOG(LogSteamVRInputDevice, Error, TEXT("[STEAMVR INPUT] %s: No Steam"), *InputAction);
+		break;
+	case VRInputError_MaxCapacityReached:
+		UE_LOG(LogSteamVRInputDevice, Error, TEXT("[STEAMVR INPUT] %s:  Max Capacity Reached"), *InputAction);
+		break;
+	case VRInputError_IPCError:
+		UE_LOG(LogSteamVRInputDevice, Error, TEXT("[STEAMVR INPUT] %s: IPC Error"), *InputAction);
+		break;
+	case VRInputError_NoActiveActionSet:
+		UE_LOG(LogSteamVRInputDevice, Error, TEXT("[STEAMVR INPUT] %s: No Active Action Set"), *InputAction);
+		break;
+	case VRInputError_InvalidDevice:
+		UE_LOG(LogSteamVRInputDevice, Error, TEXT("[STEAMVR INPUT] %s: Invalid Device"), *InputAction);
+		break;
+	case VRInputError_InvalidSkeleton:
+		UE_LOG(LogSteamVRInputDevice, Error, TEXT("[STEAMVR INPUT] %s: Invalid Skeleton"), *InputAction);
+		break;
+	case VRInputError_InvalidBoneCount:
+		UE_LOG(LogSteamVRInputDevice, Error, TEXT("[STEAMVR INPUT] %s: Invalid Bone Count"), *InputAction);
+		break;
+	case VRInputError_InvalidCompressedData:
+		UE_LOG(LogSteamVRInputDevice, Error, TEXT("[STEAMVR INPUT] %s: Invalid Compressed Data"), *InputAction);
+		break;
+	case VRInputError_NoData:
+		UE_LOG(LogSteamVRInputDevice, Error, TEXT("[STEAMVR INPUT] %s: No Data"), *InputAction);
+		break;
+	case VRInputError_BufferTooSmall:
+		UE_LOG(LogSteamVRInputDevice, Error, TEXT("[STEAMVR INPUT] %s: Buffer Too Small"), *InputAction);
+		break;
+	case VRInputError_MismatchedActionManifest:
+		UE_LOG(LogSteamVRInputDevice, Error, TEXT("[STEAMVR INPUT] %s: Mismatched Action Manifest"), *InputAction);
+		break;
+	case VRInputError_MissingSkeletonData:
+		UE_LOG(LogSteamVRInputDevice, Error, TEXT("[STEAMVR INPUT] %s: Missing Skeleton Data"), *InputAction);
+		break;
+	default:
+		UE_LOG(LogSteamVRInputDevice, Error, TEXT("[STEAMVR INPUT] %s: Unknown Error"), *InputAction);
+		break;
+	}
+
+	return;
+}
+
+void FSteamVRInputDevice::MirrorSteamVRSkeleton(VRBoneTransform_t* BoneTransformsLS, int32 BoneTransformCount) const
+{
+	check(BoneTransformsLS != nullptr);
+	check(BoneTransformCount == SteamVRSkeleton::GetBoneCount());
+
+	// Mirror the bones whose rotations transfer directly and only the translation needs to be fixed
+	{
+		const int32 TranslationOnlyBoneCount = sizeof(kMirrorTranslationOnlyBones) / sizeof(kMirrorTranslationOnlyBones[0]);
+		for (int32 i = 0; i < TranslationOnlyBoneCount; ++i)
+		{
+			const int32 BoneIndex = kMirrorTranslationOnlyBones[i];
+
+			vr::HmdVector4_t& Position = BoneTransformsLS[BoneIndex].position;
+			Position.v[0] *= -1.f;
+			Position.v[1] *= -1.f;
+			Position.v[2] *= -1.f;
+		}
+	}
+
+	// Mirror the metacarpals
+	{
+		const int32 MetaCarpalBoneCount = sizeof(kMetacarpalBones) / sizeof(kMetacarpalBones[0]);
+		for (int32 i = 0; i < MetaCarpalBoneCount; ++i)
+		{
+			const int32 BoneIndex = kMetacarpalBones[i];
+
+			vr::VRBoneTransform_t& BoneTransform = BoneTransformsLS[BoneIndex];
+
+			BoneTransform.position.v[0] *= -1.f;
+			
+			vr::HmdQuaternionf_t OriginalRotation = BoneTransform.orientation;
+			BoneTransform.orientation.w = OriginalRotation.x;
+			BoneTransform.orientation.x = -OriginalRotation.w;
+			BoneTransform.orientation.y = OriginalRotation.z;
+			BoneTransform.orientation.z = -OriginalRotation.y;
+		}
+	}
+
+	// Mirror the children of the root
+	{
+		const int32 ModelSpaceBoneCount = sizeof(kModelSpaceBones) / sizeof(kModelSpaceBones[0]);
+		for (int32 i = 0; i < ModelSpaceBoneCount; ++i)
+		{
+			const int32 BoneIndex = kModelSpaceBones[i];
+
+			vr::VRBoneTransform_t& BoneTransform = BoneTransformsLS[BoneIndex];
+			BoneTransform.position.v[0] *= -1.f;
+			BoneTransform.orientation.y *= -1.f;
+			BoneTransform.orientation.z *= -1.f;
+		}
+	}
+}
+
+bool FSteamVRInputDevice::IsVRKey(FName InputKey)
+{
+	FString KeyString = InputKey.ToString();
+
+	if ((KeyString.Contains(TEXT("SteamVR")) && !KeyString.Contains(TEXT("Generic")))
+		|| KeyString.Contains(TEXT("ValveIndex"))
+		|| KeyString.Contains(TEXT("OculusTouch"))
+		|| KeyString.Contains(TEXT("MixedReality"))
+		|| KeyString.Contains(TEXT("Vive"))
+		|| KeyString.Contains(TEXT("Cosmos"))
+		)
+	{
+		return true;
+	}
+
+	return false;
+}
+
+bool FSteamVRInputDevice::ProcessVector2D(FSteamVRAxisKeyMapping AxisKeyMapping)
+{
+	// TODO: For when Vector Actions are a thing
+	return false;
+}
+
+#undef LOCTEXT_NAMESPACE //"SteamVRInputDevice"
