Deprecate MotionController keys

diff --git a/Engine/Plugins/Runtime/Steam/SteamVR/Source/SteamVRInputDevice/Private/SteamVRInputDeviceFunctionLibrary.cpp b/Engine/Plugins/Runtime/Steam/SteamVR/Source/SteamVRInputDevice/Private/SteamVRInputDeviceFunctionLibrary.cpp
new file mode 100644
--- /dev/null
+++ b/Engine/Plugins/Runtime/Steam/SteamVR/Source/SteamVRInputDevice/Private/SteamVRInputDeviceFunctionLibrary.cpp
@@ -0,0 +1,876 @@
+/*
+Copyright 2019 Valve Corporation under https://opensource.org/licenses/BSD-3-Clause
+
+Redistribution and use in source and binary forms, with or without modification,
+are permitted provided that the following conditions are met:
+
+1. Redistributions of source code must retain the above copyright notice, this
+   list of conditions and the following disclaimer.
+
+2. Redistributions in binary form must reproduce the above copyright notice,
+   this list of conditions and the following disclaimer in the documentation
+   and/or other materials provided with the distribution.
+
+3. Neither the name of the copyright holder nor the names of its contributors
+   may be used to endorse or promote products derived from this software
+   without specific prior written permission.
+
+THIS SOFTWARE IS PROVIDED BY THE COPYRIGHT HOLDERS AND CONTRIBUTORS "AS IS" AND
+ANY EXPRESS OR IMPLIED WARRANTIES, INCLUDING, BUT NOT LIMITED TO, THE IMPLIED
+WARRANTIES OF MERCHANTABILITY AND FITNESS FOR A PARTICULAR PURPOSE ARE DISCLAIMED.
+IN NO EVENT SHALL THE COPYRIGHT HOLDER OR CONTRIBUTORS BE LIABLE FOR ANY DIRECT,
+INDIRECT, INCIDENTAL, SPECIAL, EXEMPLARY, OR CONSEQUENTIAL DAMAGES (INCLUDING,
+BUT NOT LIMITED TO, PROCUREMENT OF SUBSTITUTE GOODS OR SERVICES; LOSS OF USE, DATA,
+OR PROFITS; OR BUSINESS INTERRUPTION) HOWEVER CAUSED AND ON ANY THEORY OF LIABILITY,
+WHETHER IN CONTRACT, STRICT LIABILITY, OR TORT (INCLUDING NEGLIGENCE OR OTHERWISE)
+ARISING IN ANY WAY OUT OF THE USE OF THIS SOFTWARE, EVEN IF ADVISED OF THE
+POSSIBILITY OF SUCH DAMAGE.
+*/
+
+
+#include "SteamVRInputDeviceFunctionLibrary.h"
+
+//#if STEAMVRCONTROLLER_SUPPORTED_PLATFORMS
+#include "openvr.h"
+#include "GameFramework/WorldSettings.h"
+#include "HAL/FileManagerGeneric.h"
+using namespace vr;
+//#endif // STEAMVRCONTROLLER_SUPPORTED_PLATFORMS
+
+void USteamVRInputDeviceFunctionLibrary::PlaySteamVR_HapticFeedback(ESteamVRHand Hand, float StartSecondsFromNow, float DurationSeconds, float Frequency, float Amplitude)
+{
+	FSteamVRInputDevice* SteamVRInputDevice = GetSteamVRInputDevice();
+	if (VRSystem() && VRInput())
+	{
+		if (Amplitude < 0.f)
+		{
+			Amplitude = 0.f;
+		}
+		else if (Amplitude > 1.f)
+		{
+			Amplitude = 1.f;
+		}
+
+		// Setup which hand data we will get from SteamVR
+		VRActionHandle_t ActiveSkeletalHand = k_ulInvalidActionHandle;
+
+		if (Hand == ESteamVRHand::VR_Left && SteamVRInputDevice->VRVibrationLeft != k_ulInvalidActionHandle && SteamVRInputDevice->bIsSkeletalControllerLeftPresent && SteamVRInputDevice->VRSkeletalHandleLeft != k_ulInvalidActionHandle)
+		{
+			if (SteamVRInputDevice->VRVibrationLeft == k_ulInvalidActionHandle)
+			{
+				return;
+			}
+
+			ActiveSkeletalHand = SteamVRInputDevice->VRSkeletalHandleLeft;
+			VRInput()->TriggerHapticVibrationAction(SteamVRInputDevice->VRVibrationLeft, StartSecondsFromNow, DurationSeconds, Frequency, Amplitude, k_ulInvalidInputValueHandle);
+		}
+		else if (Hand == ESteamVRHand::VR_Right && SteamVRInputDevice->VRVibrationLeft != k_ulInvalidActionHandle && SteamVRInputDevice->bIsSkeletalControllerRightPresent && SteamVRInputDevice->VRSkeletalHandleRight != k_ulInvalidActionHandle)
+		{
+			if (SteamVRInputDevice->VRVibrationRight == k_ulInvalidActionHandle)
+			{
+				return;
+			}
+
+			ActiveSkeletalHand = SteamVRInputDevice->VRSkeletalHandleRight;
+			VRInput()->TriggerHapticVibrationAction(SteamVRInputDevice->VRVibrationRight, StartSecondsFromNow, DurationSeconds, Frequency, Amplitude, k_ulInvalidInputValueHandle);
+		}
+	}
+}
+
+void USteamVRInputDeviceFunctionLibrary::RegenActionManifest()
+{
+#if WITH_EDITOR
+	FSteamVRInputDevice* SteamVRInputDevice = GetSteamVRInputDevice();
+	if (SteamVRInputDevice != nullptr)
+	{
+		SteamVRInputDevice->RegenerateActionManifest();
+	}
+#endif
+}
+
+void USteamVRInputDeviceFunctionLibrary::RegenControllerBindings()
+{
+#if WITH_EDITOR
+	FSteamVRInputDevice* SteamVRInputDevice = GetSteamVRInputDevice();
+	if (SteamVRInputDevice != nullptr)
+	{
+		SteamVRInputDevice->RegenerateControllerBindings();
+	}
+#endif
+}
+
+void USteamVRInputDeviceFunctionLibrary::ReloadActionManifest()
+{
+#if WITH_EDITOR
+	FSteamVRInputDevice* SteamVRInputDevice = GetSteamVRInputDevice();
+	if (SteamVRInputDevice != nullptr)
+	{
+		SteamVRInputDevice->ReloadActionManifest();
+	}
+#endif
+}
+
+void USteamVRInputDeviceFunctionLibrary::GetCurlsAndSplaysState(bool& LeftHandState, bool& RightHandState)
+{
+	FSteamVRInputDevice* SteamVRInputDevice = GetSteamVRInputDevice();
+	if (SteamVRInputDevice != nullptr)
+	{
+		LeftHandState = SteamVRInputDevice->bCurlsAndSplaysEnabled_L;
+		RightHandState = SteamVRInputDevice->bCurlsAndSplaysEnabled_R;
+	}
+}
+
+void USteamVRInputDeviceFunctionLibrary::GetSkeletalState(bool& LeftHandState, bool& RightHandState)
+{
+	FSteamVRInputDevice* SteamVRInputDevice = GetSteamVRInputDevice();
+	if (SteamVRInputDevice != nullptr)
+	{
+		LeftHandState = SteamVRInputDevice->bIsSkeletalControllerLeftPresent;
+		RightHandState = SteamVRInputDevice->bIsSkeletalControllerRightPresent;
+	}
+}
+
+void USteamVRInputDeviceFunctionLibrary::GetControllerFidelity(EControllerFidelity& LeftControllerFidelity, EControllerFidelity& RightControllerFidelity)
+{
+	FSteamVRInputDevice* SteamVRInputDevice = GetSteamVRInputDevice();
+	if (SteamVRInputDevice != nullptr)
+	{
+		// Update the skeletal tracking level from SteamVR
+		SteamVRInputDevice->GetControllerFidelity();
+
+		// Return the left controller fidelity
+		switch (SteamVRInputDevice->LeftControllerFidelity)
+		{
+		case EVRSkeletalTrackingLevel::VRSkeletalTracking_Full:
+			LeftControllerFidelity = EControllerFidelity::VR_ControllerFidelity_Full;
+			break;
+
+		case EVRSkeletalTrackingLevel::VRSkeletalTracking_Partial:
+			LeftControllerFidelity = EControllerFidelity::VR_ControllerFidelity_Partial;
+			break;
+
+		case EVRSkeletalTrackingLevel::VRSkeletalTracking_Estimated:
+			// falls through
+		default:
+			LeftControllerFidelity = EControllerFidelity::VR_ControllerFidelity_Estimated;
+			break;
+		}
+
+		// Return the right controller fidelity
+		switch (SteamVRInputDevice->RightControllerFidelity)
+		{
+		case EVRSkeletalTrackingLevel::VRSkeletalTracking_Full:
+			RightControllerFidelity = EControllerFidelity::VR_ControllerFidelity_Full;
+			break;
+
+		case EVRSkeletalTrackingLevel::VRSkeletalTracking_Partial:
+			RightControllerFidelity = EControllerFidelity::VR_ControllerFidelity_Partial;
+			break;
+
+		case EVRSkeletalTrackingLevel::VRSkeletalTracking_Estimated:
+			// falls through
+		default:
+			RightControllerFidelity = EControllerFidelity::VR_ControllerFidelity_Estimated;
+			break;
+		}
+	}
+}
+
+void USteamVRInputDeviceFunctionLibrary::SetCurlsAndSplaysState(bool NewLeftHandState, bool NewRightHandState)
+{
+	FSteamVRInputDevice* SteamVRInputDevice = GetSteamVRInputDevice();
+	if (SteamVRInputDevice != nullptr)
+	{
+		SteamVRInputDevice->bCurlsAndSplaysEnabled_L = NewLeftHandState;
+		SteamVRInputDevice->bCurlsAndSplaysEnabled_R = NewRightHandState;
+	}
+}
+
+void USteamVRInputDeviceFunctionLibrary::GetPoseSource(bool& bUsingSkeletonPose)
+{
+	FSteamVRInputDevice* SteamVRInputDevice = GetSteamVRInputDevice();
+	if (SteamVRInputDevice != nullptr)
+	{
+		bUsingSkeletonPose = SteamVRInputDevice->bUseSkeletonPose;
+	}
+}
+
+void USteamVRInputDeviceFunctionLibrary::SetPoseSource(bool bUseSkeletonPose)
+{
+	FSteamVRInputDevice* SteamVRInputDevice = GetSteamVRInputDevice();
+	if (SteamVRInputDevice != nullptr)
+	{
+		SteamVRInputDevice->bUseSkeletonPose = bUseSkeletonPose;
+	}
+}
+
+void USteamVRInputDeviceFunctionLibrary::GetSkeletalTransform(FSteamVRSkeletonTransform& LeftHand, FSteamVRSkeletonTransform& RightHand, bool bWithController)
+{
+	FTransform OutPose[STEAMVR_SKELETON_BONE_COUNT];
+
+	FSteamVRInputDevice* SteamVRInputDevice = GetSteamVRInputDevice();
+	if (SteamVRInputDevice != nullptr)
+	{
+		// Setup Motion Range
+		EVRSkeletalMotionRange MotionRange = bWithController ? VRSkeletalMotionRange_WithController : VRSkeletalMotionRange_WithoutController;
+
+		// Left Hand - Grab Skeletal Data
+		SteamVRInputDevice->GetSkeletalData(true, false, MotionRange, OutPose, STEAMVR_SKELETON_BONE_COUNT);
+
+		LeftHand.Root = OutPose[0];
+		LeftHand.Wrist = OutPose[1];
+
+		LeftHand.Thumb_0 = OutPose[2];
+		LeftHand.Thumb_1 = OutPose[3];
+		LeftHand.Thumb_2 = OutPose[4];
+		LeftHand.Thumb_3 = OutPose[5];
+
+		LeftHand.Index_0 = OutPose[6];
+		LeftHand.Index_1 = OutPose[7];
+		LeftHand.Index_2 = OutPose[8];
+		LeftHand.Index_3 = OutPose[9];
+		LeftHand.Index_4 = OutPose[10];
+
+		LeftHand.Middle_0 = OutPose[11];
+		LeftHand.Middle_1 = OutPose[12];
+		LeftHand.Middle_2 = OutPose[13];
+		LeftHand.Middle_3 = OutPose[14];
+		LeftHand.Middle_4 = OutPose[15];
+
+		LeftHand.Ring_0 = OutPose[16];
+		LeftHand.Ring_1 = OutPose[17];
+		LeftHand.Ring_2 = OutPose[18];
+		LeftHand.Ring_3 = OutPose[19];
+		LeftHand.Ring_4 = OutPose[20];
+
+		LeftHand.Pinky_0 = OutPose[21];
+		LeftHand.Pinky_1 = OutPose[22];
+		LeftHand.Pinky_2 = OutPose[23];
+		LeftHand.Pinky_3 = OutPose[24];
+		LeftHand.Pinky_4 = OutPose[25];
+
+		LeftHand.Aux_Thumb = OutPose[26];
+		LeftHand.Aux_Index = OutPose[27];
+		LeftHand.Aux_Middle = OutPose[28];
+		LeftHand.Aux_Ring = OutPose[29];
+		LeftHand.Aux_Pinky = OutPose[30];
+
+		LeftHand.Bone_Count = OutPose[31];
+
+		// Right Hand - Grab Skeletal Data
+		SteamVRInputDevice->GetSkeletalData(false, false, MotionRange, OutPose, STEAMVR_SKELETON_BONE_COUNT);
+
+		RightHand.Root = OutPose[0];
+		RightHand.Wrist = OutPose[1];
+
+		RightHand.Thumb_0 = OutPose[2];
+		RightHand.Thumb_1 = OutPose[3];
+		RightHand.Thumb_2 = OutPose[4];
+		RightHand.Thumb_3 = OutPose[5];
+
+		RightHand.Index_0 = OutPose[6];
+		RightHand.Index_1 = OutPose[7];
+		RightHand.Index_2 = OutPose[8];
+		RightHand.Index_3 = OutPose[9];
+		RightHand.Index_4 = OutPose[10];
+
+		RightHand.Middle_0 = OutPose[11];
+		RightHand.Middle_1 = OutPose[12];
+		RightHand.Middle_2 = OutPose[13];
+		RightHand.Middle_3 = OutPose[14];
+		RightHand.Middle_4 = OutPose[15];
+
+		RightHand.Ring_0 = OutPose[16];
+		RightHand.Ring_1 = OutPose[17];
+		RightHand.Ring_2 = OutPose[18];
+		RightHand.Ring_3 = OutPose[19];
+		RightHand.Ring_4 = OutPose[20];
+
+		RightHand.Pinky_0 = OutPose[21];
+		RightHand.Pinky_1 = OutPose[22];
+		RightHand.Pinky_2 = OutPose[23];
+		RightHand.Pinky_3 = OutPose[24];
+		RightHand.Pinky_4 = OutPose[25];
+
+		RightHand.Aux_Thumb = OutPose[26];
+		RightHand.Aux_Index = OutPose[27];
+		RightHand.Aux_Middle = OutPose[28];
+		RightHand.Aux_Ring = OutPose[29];
+		RightHand.Aux_Pinky = OutPose[30];
+
+		RightHand.Bone_Count = OutPose[31];
+	}
+}
+
+void USteamVRInputDeviceFunctionLibrary::GetLeftHandPoseData(FVector& Position, FRotator& Orientation, FVector& AngularVelocity, FVector& Velocity)
+{
+	FSteamVRInputDevice* SteamVRInputDevice = GetSteamVRInputDevice();
+	if (SteamVRInputDevice != nullptr)
+	{
+		SteamVRInputDevice->GetLeftHandPoseData(Position, Orientation, AngularVelocity, Velocity);
+	}
+}
+
+void USteamVRInputDeviceFunctionLibrary::GetRightHandPoseData(FVector& Position, FRotator& Orientation, FVector& AngularVelocity, FVector& Velocity)
+{
+	FSteamVRInputDevice* SteamVRInputDevice = GetSteamVRInputDevice();
+	if (SteamVRInputDevice != nullptr)
+	{
+		SteamVRInputDevice->GetRightHandPoseData(Position, Orientation, AngularVelocity, Velocity);
+	}
+}
+
+void USteamVRInputDeviceFunctionLibrary::GetSteamVR_ActionArray(TArray<FSteamVRAction>& SteamVRActions)
+{
+	bool bAlreadyExists;
+	FSteamVRInputDevice* SteamVRInputDevice = GetSteamVRInputDevice();
+	if (SteamVRInputDevice != nullptr)
+	{
+		// Get the input actions and assign them to the defined actions list for developer consumption
+		for (FSteamVRInputAction& InputAction : SteamVRInputDevice->ActionEvents)
+		{
+			// Check for duplicates
+			bAlreadyExists = false;
+			for (FSteamVRAction& SteamVRAction : SteamVRActions)
+			{
+				if (SteamVRAction.Handle == InputAction.Handle)
+				{
+					bAlreadyExists = true;
+					break;
+				}
+			}
+
+			// Add the Input Action to the list of Defined Actions if it's unique
+			if (!bAlreadyExists)
+			{
+				SteamVRActions.Add(FSteamVRAction(InputAction.Name, InputAction.Path, InputAction.Handle, InputAction.ActiveOrigin));
+			}
+		}
+	}
+}
+
+void USteamVRInputDeviceFunctionLibrary::FindSteamVR_Action(FName ActionName, bool& bResult, FSteamVRAction& FoundAction, FSteamVRActionSet& FoundActionSet, FName ActionSet)
+{
+	bool bActionFound = false;
+	bool bActionSetFound = false;
+
+	// Get current SteamVR Actions
+	TArray<FSteamVRAction> SteamVRActions;
+	GetSteamVR_ActionArray(SteamVRActions);
+
+	// Find Action
+	for (FSteamVRAction SteamVRDefinedAction : SteamVRActions)
+	{
+		if (SteamVRDefinedAction.Name.IsEqual(ActionName))
+		{
+			FoundAction = SteamVRDefinedAction;
+			bActionFound = true;
+			break;
+		}
+	}
+
+	// Get current SteamVR Action Sets
+	FString InActionSet = TEXT("/actions/") + ActionSet.ToString();
+	TArray<FSteamVRActionSet> SteamVRActionSets;
+	GetSteamVR_ActionSetArray(SteamVRActionSets);
+
+	// Find Action
+	for (FSteamVRActionSet SteamVRDefinedActionSet : SteamVRActionSets)
+	{
+		if (SteamVRDefinedActionSet.Path.Equals(InActionSet, ESearchCase::IgnoreCase))
+		{
+			FoundActionSet = SteamVRDefinedActionSet;
+			bActionSetFound = true;
+			break;
+		}
+	}
+
+	// Update search result 
+	bResult = bActionFound && bActionSetFound;
+}
+
+void USteamVRInputDeviceFunctionLibrary::GetSteamVR_ActionSetArray(TArray<FSteamVRActionSet>& SteamVRActionSets)
+{
+	// We only have one Action Set supported for this version of the plugin, so just return that
+	FSteamVRInputDevice* SteamVRInputDevice = GetSteamVRInputDevice();
+	if (SteamVRInputDevice != nullptr)
+	{
+		SteamVRActionSets.Add(FSteamVRActionSet(TEXT(ACTION_SET), SteamVRInputDevice->MainActionSet));
+	}
+}
+
+bool USteamVRInputDeviceFunctionLibrary::GetSteamVR_OriginTrackedDeviceInfo(FSteamVRAction SteamVRAction, FSteamVRInputOriginInfo& InputOriginInfo)
+{
+	if (VRSystem() && VRInput())
+	{
+		InputOriginInfo_t OriginInfo = {};
+		EVRInputError Err = VRInput()->GetOriginTrackedDeviceInfo(SteamVRAction.ActiveOrigin, &OriginInfo, sizeof(OriginInfo));
+
+		if (Err == VRInputError_None && OriginInfo.trackedDeviceIndex != k_unTrackedDeviceIndexInvalid)
+		{
+			// Get device model information
+			char ModelBuffer[k_unMaxPropertyStringSize];
+			uint32 StringBytes = VRSystem()->GetStringTrackedDeviceProperty(OriginInfo.trackedDeviceIndex, ETrackedDeviceProperty::Prop_ModelNumber_String, ModelBuffer, sizeof(ModelBuffer));
+
+			// Set Input Origin Info
+			InputOriginInfo.TrackedDeviceIndex = OriginInfo.trackedDeviceIndex;
+			InputOriginInfo.RenderModelComponentName = *FString(UTF8_TO_TCHAR(OriginInfo.rchRenderModelComponentName));
+			InputOriginInfo.TrackedDeviceModel = *FString(UTF8_TO_TCHAR(ModelBuffer));
+			return true;
+		}
+		else
+		{
+			UE_LOG(LogTemp, Warning, TEXT("Error [%i] trying to retrieve Tracked Device Info or an invalid device was returned for Action [%s]"), (int)Err, *SteamVRAction.Path);
+		}
+	}
+
+	return false;
+}
+
+void USteamVRInputDeviceFunctionLibrary::FindSteamVR_OriginTrackedDeviceInfo(FName ActionName, bool& bResult, FSteamVRInputOriginInfo& InputOriginInfo, FName ActionSet /*= FName("main")*/)
+{
+	bResult = false;
+	FSteamVRAction FoundAction;
+	FSteamVRActionSet FoundActionSet;
+
+	FindSteamVR_Action(ActionName, bResult, FoundAction, FoundActionSet, ActionSet);
+
+	if (bResult)
+	{
+		bResult = GetSteamVR_OriginTrackedDeviceInfo(FoundAction, InputOriginInfo);
+	}
+}
+
+void USteamVRInputDeviceFunctionLibrary::GetSteamVR_OriginLocalizedName(FSteamVRAction SteamVRAction, TArray<ESteamVRInputStringBits> LocalizedParts, FString& OriginLocalizedName)
+{
+	if (VRSystem() && VRInput())
+	{
+		uint32 LocalizedPartsMask = 0;
+		EVRInputStringBits SteamVREquivalentInputStringBits;
+
+		for (ESteamVRInputStringBits LocalizedPart : LocalizedParts)
+		{
+			switch (LocalizedPart)
+			{
+			case ESteamVRInputStringBits::VR_InputString_Hand:
+				SteamVREquivalentInputStringBits = EVRInputStringBits::VRInputString_Hand;
+				break;
+
+			case ESteamVRInputStringBits::VR_InputString_ControllerType:
+				SteamVREquivalentInputStringBits = EVRInputStringBits::VRInputString_ControllerType;
+				break;
+
+			case ESteamVRInputStringBits::VR_InputString_InputSource:
+				SteamVREquivalentInputStringBits = EVRInputStringBits::VRInputString_InputSource;
+				break;
+
+			case ESteamVRInputStringBits::VR_InputString_All:
+				// Falls through
+
+			default:
+				SteamVREquivalentInputStringBits = EVRInputStringBits::VRInputString_All;
+				break;
+			}
+
+			LocalizedPartsMask |= (int)SteamVREquivalentInputStringBits;
+		}
+
+		// Retrieve Localized Name
+		char buf[k_unMaxPropertyStringSize];
+		EVRInputError Err = VRInput()->GetOriginLocalizedName(SteamVRAction.ActiveOrigin, buf, sizeof(buf), LocalizedPartsMask);
+		OriginLocalizedName = *FString(UTF8_TO_TCHAR(buf));
+
+		// Provide debugging info if retrieval is unsuccessful
+		//if (Err != VRInputError_None)
+		//{
+		//	UE_LOG(LogTemp, Warning, TEXT("Error [%i] when retrieving Origin Localized Name for Action [%i] %s"), (int)Err, (int)SteamVRDefinedAction.Handle, *SteamVRDefinedAction.Path);
+		//}
+	}
+}
+
+void USteamVRInputDeviceFunctionLibrary::ShowSteamVR_ActionOrigin(FSteamVRAction SteamVRAction, FSteamVRActionSet SteamVRActionSet)
+{
+	if (VRSystem() && VRInput())
+	{
+		// Show the action origin in user's hmd
+		EVRInputError Err = VRInput()->ShowActionOrigins(0, SteamVRAction.Handle);
+		FSteamVRInputOriginInfo OriginInfo;
+
+		if (GetSteamVR_OriginTrackedDeviceInfo(SteamVRAction, OriginInfo))
+		{
+			VRActiveActionSet_t ActiveActionSets[] = { 0 };
+			ActiveActionSets[0].ulActionSet = SteamVRActionSet.Handle;
+			VRInput()->ShowBindingsForActionSet(ActiveActionSets, sizeof(ActiveActionSets[0]), 1, SteamVRAction.ActiveOrigin);
+
+			UE_LOG(LogTemp, Warning, TEXT("Action [%s] triggered from Device [%i][%s] at Component [%s]"), *SteamVRAction.Name.ToString(), OriginInfo.TrackedDeviceIndex, *OriginInfo.TrackedDeviceModel, *OriginInfo.RenderModelComponentName);
+		}
+	}
+}
+
+bool USteamVRInputDeviceFunctionLibrary::FindSteamVR_ActionOrigin(FName ActionName, FName ActionSet)
+{
+	// Find SteamVR Action
+	bool bIsActionFound = false;
+	FSteamVRAction SteamVRAction;
+	FSteamVRActionSet SteamVRActionSet;
+	FindSteamVR_Action(ActionName, bIsActionFound, SteamVRAction, SteamVRActionSet, ActionSet);
+
+	if (bIsActionFound && VRSystem() && VRInput())
+	{
+		ShowSteamVR_ActionOrigin(SteamVRAction, SteamVRActionSet);
+		return true;
+	}
+
+	GEngine->AddOnScreenDebugMessage(-1, 5.f, FColor::Red, (TEXT("Unable to find Action [%s] for Action Set [%s]"), *ActionName.ToString(), *ActionSet.ToString()));
+	return false;
+}
+
+bool USteamVRInputDeviceFunctionLibrary::GetSteamVR_HandPoseRelativeToNow(FVector& Position, FRotator& Orientation, ESteamVRHand Hand /*= ESteamVRHand::VR_Left*/, float PredictedSecondsFromNow /*= 0.f*/)
+{
+	if (VRSystem() && VRInput())
+	{
+		FSteamVRInputDevice* SteamVRInputDevice = GetSteamVRInputDevice();
+		if (SteamVRInputDevice != nullptr)
+		{
+			VRActionHandle_t HandActionHandle = (Hand == ESteamVRHand::VR_Left) ? SteamVRInputDevice->VRControllerHandleLeft : SteamVRInputDevice->VRControllerHandleRight;
+			if (HandActionHandle != k_ulInvalidActionHandle)
+			{
+				InputPoseActionData_t PoseData = { 0 };
+				EVRInputError InputError = VRInput()->GetPoseActionDataRelativeToNow(HandActionHandle, VRCompositor()->GetTrackingSpace(), PredictedSecondsFromNow, &PoseData, sizeof(PoseData), k_ulInvalidInputValueHandle);
+
+				if (InputError == VRInputError_None)
+				{
+					// Get SteamVR Transform Matrix for this controller
+					HmdMatrix34_t Matrix = PoseData.pose.mDeviceToAbsoluteTracking;
+
+					// Transform SteamVR Pose to Unreal Pose
+					FMatrix Pose = FMatrix(
+						FPlane(Matrix.m[0][0], Matrix.m[1][0], Matrix.m[2][0], 0.0f),
+						FPlane(Matrix.m[0][1], Matrix.m[1][1], Matrix.m[2][1], 0.0f),
+						FPlane(Matrix.m[0][2], Matrix.m[1][2], Matrix.m[2][2], 0.0f),
+						FPlane(Matrix.m[0][3], Matrix.m[1][3], Matrix.m[2][3], 1.0f)
+					);
+
+
+					// Transform SteamVR Rotation Quaternion to a UE FRotator
+					FQuat OrientationQuat;
+					FQuat OrientationPose(Pose);
+					OrientationQuat.X = -OrientationPose.Z;
+					OrientationQuat.Y = OrientationPose.X;
+					OrientationQuat.Z = OrientationPose.Y;
+					OrientationQuat.W = -OrientationPose.W;
+
+					FVector PositionPose = ((FVector(-Pose.M[3][2], Pose.M[3][0], Pose.M[3][1])) * GWorld->GetWorldSettings()->WorldToMeters);
+					Position = PositionPose;
+
+					//OutOrientation = BaseOrientation.Inverse() * OutOrientation;
+					OrientationQuat.Normalize();
+					Orientation = OrientationQuat.Rotator();
+					Orientation.Normalize();
+
+					return true;
+				}
+			}
+		}
+	}
+
+	return false;
+}
+
+float USteamVRInputDeviceFunctionLibrary::GetSteamVR_GlobalPredictedSecondsFromNow()
+{
+	if (VRSystem() && VRInput())
+	{
+		FSteamVRInputDevice* SteamVRInputDevice = GetSteamVRInputDevice();
+		if (SteamVRInputDevice != nullptr)
+		{
+			return SteamVRInputDevice->GlobalPredictedSecondsFromNow;
+		}
+	}
+
+	return -9999.f;
+}
+
+float USteamVRInputDeviceFunctionLibrary::SetSteamVR_GlobalPredictedSecondsFromNow(float NewValue)
+{
+	if (VRSystem() && VRInput())
+	{
+		FSteamVRInputDevice* SteamVRInputDevice = GetSteamVRInputDevice();
+		if (SteamVRInputDevice != nullptr)
+		{
+			SteamVRInputDevice->GlobalPredictedSecondsFromNow = NewValue;
+			return SteamVRInputDevice->GlobalPredictedSecondsFromNow;
+		}
+	}
+
+	return -9999.f;
+}
+
+void USteamVRInputDeviceFunctionLibrary::ShowAllSteamVR_ActionOrigins()
+{
+	VRActiveActionSet_t ActiveActionSets[1];
+	VRInput()->ShowBindingsForActionSet(ActiveActionSets, sizeof(ActiveActionSets[0]), 0, 0);
+}
+
+TArray<FSteamVRInputBindingInfo> USteamVRInputDeviceFunctionLibrary::GetSteamVR_InputBindingInfo(FSteamVRAction SteamVRActionHandle)
+{
+	// Setup variables to be used by the OpenVR call
+	TArray<FSteamVRInputBindingInfo> SteamVRBindingInfo;
+	InputBindingInfo_t* InputBindingInfo = new InputBindingInfo_t[MAX_BINDINGINFO_COUNT];
+	uint32_t BindingInfoCount = 0;
+	EVRInputError BindingInfoError = VRInputError_NoData;
+
+	// Execute OpenVR call GetActionBindingInfo if action handle is valid and VRInput is present
+	if (VRInput() && !SteamVRActionHandle.Path.IsEmpty())
+	{
+		// Get binding info for provided action handle in the currently active controller type
+		BindingInfoError = VRInput()->GetActionBindingInfo(SteamVRActionHandle.Handle, InputBindingInfo, sizeof(InputBindingInfo_t), MAX_BINDINGINFO_COUNT, &BindingInfoCount);
+
+		// Check if call is successful
+		if (BindingInfoError != VRInputError_None)
+		{
+			UE_LOG(LogTemp, Error, TEXT("Error retrieving binding info for active controller: %i"), (int)BindingInfoError);
+		}
+		else
+		{
+			// Check if there are binding info related to the action
+			if (BindingInfoCount > 0)
+			{
+				// If there are, go through each binding info found
+				for (int32 i = 0; i < (int32)BindingInfoCount; i++)
+				{
+					// Convert binding info text from a char array to a UE4 FString
+					FString DevicePathName = FString(UTF8_TO_TCHAR(InputBindingInfo[i].rchDevicePathName));
+					FString InputPathName = FString(UTF8_TO_TCHAR(InputBindingInfo[i].rchInputPathName));
+					FString ModeName = FString(UTF8_TO_TCHAR(InputBindingInfo[i].rchModeName));
+					FString SlotName = FString(UTF8_TO_TCHAR(InputBindingInfo[i].rchSlotName));
+
+					// Create a BP compatible binding info struct and add this to the output array
+					SteamVRBindingInfo.Add(FSteamVRInputBindingInfo(
+						FName(*DevicePathName),
+						FName(*InputPathName),
+						FName(*ModeName),
+						FName(*SlotName)
+					));
+
+					//UE_LOG(LogTemp, Error, TEXT("Found Binding Info: %s %s %s %s"), *DevicePathName, *InputPathName, *ModeName, *SlotName);
+				}
+			}
+			else
+			{
+				UE_LOG(LogTemp, Warning, TEXT("No binding info found for action: %s"), *SteamVRActionHandle.Path);
+			}
+		}
+	}
+	else
+	{
+		UE_LOG(LogTemp, Error, TEXT("Error calling GetSteamVRInputBindingInfo from Blueprint. OpenVRInput is not present or SteamVR Action Handle provided as input is invalid or empty!"));
+	}
+
+	return SteamVRBindingInfo;
+}
+
+TArray<FSteamVRInputBindingInfo> USteamVRInputDeviceFunctionLibrary::FindSteamVR_InputBindingInfo(FName ActionName, FName ActionSet /*= FName("main")*/)
+{
+	// Check for a valid OpenVR session and action name/set
+	if (VRInput() && ActionName != NAME_None && ActionSet != NAME_None)
+	{
+		// Get the action handle for given action name and action set
+		bool bFindResult = false;
+		FSteamVRAction FoundAction;
+		FSteamVRActionSet FoundActionSet;
+		FindSteamVR_Action(ActionName, bFindResult, FoundAction, FoundActionSet, ActionSet);
+
+		// Check if we found the action
+		if (bFindResult)
+		{
+			// If found, get the binding info and return the result
+			return GetSteamVR_InputBindingInfo(FoundAction);
+		}
+		else
+		{
+			// Otherwise, provide an error message
+			UE_LOG(LogTemp, Error, TEXT("Error when calling FindSteamVRInputBindingInfo, there was no Action found with the name %s in action set %s"), *ActionName.ToString(), *ActionSet.ToString());
+		}
+
+	}
+	else
+	{
+		UE_LOG(LogTemp, Error, TEXT("Error calling FindSteamVRInputBindingInfo from Blueprint. OpenVRInput is not present or SteamVR action name/set provided is empty!"));
+	}
+
+	// Return an empty array of binding info
+	TArray<FSteamVRInputBindingInfo> EmptyInfo;
+	return EmptyInfo;
+
+}
+
+bool USteamVRInputDeviceFunctionLibrary::ResetSeatedPosition()
+{
+	if (VRSystem() && VRInput())
+	{
+		if (VRCompositor()->GetTrackingSpace() == TrackingUniverseSeated)
+		{
+			VRSystem()->ResetSeatedZeroPose();
+			return true;
+		}
+		else
+		{
+			UE_LOG(LogTemp, Warning, TEXT("Request to Reset Seated Position was ignored. Use \"Set Tracking Origin\" to \"Eye\" before calling this function."));
+		}
+	}
+
+	return false;
+}
+
+
+float USteamVRInputDeviceFunctionLibrary::GetUserIPD()
+{
+	if (VRSystem())
+	{
+		return VRSystem()->GetFloatTrackedDeviceProperty(k_unTrackedDeviceIndex_Hmd, ETrackedDeviceProperty::Prop_UserIpdMeters_Float) * 1000; // Return IPD in mm
+	}
+
+	return 0.f;
+}
+
+FTransform USteamVRInputDeviceFunctionLibrary::GetUETransform(VRBoneTransform_t SteamBoneTransform)
+{
+	FTransform RetTransform;
+
+	// Calculate UE Bone Transform
+	FQuat OrientationQuat(SteamBoneTransform.orientation.x,
+		-SteamBoneTransform.orientation.y,
+		SteamBoneTransform.orientation.z,
+		-SteamBoneTransform.orientation.w);
+	OrientationQuat.Normalize();
+	RetTransform = FTransform(OrientationQuat,
+		FVector(SteamBoneTransform.position.v[2] * -1.f,
+			SteamBoneTransform.position.v[0],
+			SteamBoneTransform.position.v[1]) * 0.1f);
+	return RetTransform;
+}
+
+void USteamVRInputDeviceFunctionLibrary::LaunchBindingsURL()
+{
+	FSteamVRInputDevice* SteamVRInputDevice = GetSteamVRInputDevice();
+	if (SteamVRInputDevice != nullptr)
+	{
+		FString BindingsURL = "http://localhost:27062/dashboard/controllerbinding.html?app=" + SteamVRInputDevice->EditorAppKey;
+		UE_LOG(LogTemp, Warning, TEXT("%s"), *BindingsURL);
+		FPlatformProcess::LaunchURL(*BindingsURL, nullptr, nullptr);
+	}
+}
+
+void USteamVRInputDeviceFunctionLibrary::GetFingerCurlsAndSplays(EHand Hand, FSteamVRFingerCurls& FingerCurls, FSteamVRFingerSplays& FingerSplays, ESkeletalSummaryDataType SummaryDataType)
+{
+	FSteamVRInputDevice* SteamVRInputDevice = GetSteamVRInputDevice();
+	if (SteamVRInputDevice != nullptr && VRSystem() && VRInput())
+	{
+		// Get action state this frame
+		VRActiveActionSet_t ActiveActionSets[] = {
+			{
+				SteamVRInputDevice->MainActionSet,
+				k_ulInvalidInputValueHandle,
+				k_ulInvalidActionSetHandle
+			}
+		};
+
+		EVRInputError UpdateActionStateError = VRInput()->UpdateActionState(ActiveActionSets, sizeof(VRActiveActionSet_t), 1);
+		if (UpdateActionStateError != VRInputError_None)
+		{
+			FingerCurls = {};
+			FingerSplays = {};
+			return;
+		}
+
+		// Setup which hand data we will get from SteamVR
+		VRActionHandle_t ActiveSkeletalHand = k_ulInvalidActionHandle;
+
+		if (Hand == EHand::VR_LeftHand && SteamVRInputDevice->bIsSkeletalControllerLeftPresent && SteamVRInputDevice->VRSkeletalHandleLeft != k_ulInvalidActionHandle)
+		{
+			ActiveSkeletalHand = SteamVRInputDevice->VRSkeletalHandleLeft;
+		}
+		else if (Hand == EHand::VR_RightHand && SteamVRInputDevice->bIsSkeletalControllerRightPresent && SteamVRInputDevice->VRSkeletalHandleRight != k_ulInvalidActionHandle)
+		{
+			ActiveSkeletalHand = SteamVRInputDevice->VRSkeletalHandleRight;
+		}
+
+		if (ActiveSkeletalHand == k_ulInvalidActionHandle)
+		{
+			FingerCurls = {};
+			FingerSplays = {};
+			return;
+		}
+
+		InputSkeletalActionData_t actionData;
+		EVRInputError GetSkeletalActionDataError = VRInput()->GetSkeletalActionData(ActiveSkeletalHand, &actionData, sizeof(InputSkeletalActionData_t));
+
+		if (GetSkeletalActionDataError != VRInputError_None)
+		{
+			FingerCurls = {};
+			FingerSplays = {};
+			return;
+		}
+
+		if (actionData.bActive == false)
+		{
+			FingerCurls = {};
+			FingerSplays = {};
+			return;
+		}
+
+		// Get Skeletal Summary Data
+		VRSkeletalSummaryData_t ActiveSkeletalSummaryData;
+		if (ActiveSkeletalHand == k_ulInvalidActionHandle)
+		{
+			FingerCurls = {};
+			FingerSplays = {};
+			return;
+		}
+
+		EVRSummaryType SteamVRSummaryType = (SummaryDataType == ESkeletalSummaryDataType::VR_SummaryType_FromDevice) ? VRSummaryType_FromDevice : VRSummaryType_FromAnimation;
+		EVRInputError GetSkeletalSummaryDataError = VRInput()->GetSkeletalSummaryData(ActiveSkeletalHand, SteamVRSummaryType, &ActiveSkeletalSummaryData);
+
+		if (GetSkeletalSummaryDataError != VRInputError_None)
+		{
+			FingerCurls = {};
+			FingerSplays = {};
+			return;
+		}
+		else if (GetSkeletalSummaryDataError == VRInputError_None)
+		{
+			// Update curls and splay values for output
+			FingerCurls.Thumb = ActiveSkeletalSummaryData.flFingerCurl[VRFinger_Thumb];
+			FingerCurls.Index = ActiveSkeletalSummaryData.flFingerCurl[VRFinger_Index];
+			FingerCurls.Middle = ActiveSkeletalSummaryData.flFingerCurl[VRFinger_Middle];
+			FingerCurls.Ring = ActiveSkeletalSummaryData.flFingerCurl[VRFinger_Ring];
+			FingerCurls.Pinky = ActiveSkeletalSummaryData.flFingerCurl[VRFinger_Pinky];
+
+			FingerSplays.Index_Middle = ActiveSkeletalSummaryData.flFingerSplay[VRFingerSplay_Index_Middle];
+			FingerSplays.Middle_Ring = ActiveSkeletalSummaryData.flFingerSplay[VRFingerSplay_Middle_Ring];
+			FingerSplays.Ring_Pinky = ActiveSkeletalSummaryData.flFingerSplay[VRFingerSplay_Ring_Pinky];
+			FingerSplays.Thumb_Index = ActiveSkeletalSummaryData.flFingerSplay[VRFingerSplay_Thumb_Index];
+			return;
+		}
+	}
+
+	// Unable to retrieve the curls and splay values for this hand, send zeroed out values back to the user
+	FingerCurls = {};
+	FingerSplays = {};
+}
+
+FSteamVRInputDevice* USteamVRInputDeviceFunctionLibrary::GetSteamVRInputDevice()
+{
+	TArray<IMotionController*> MotionControllers = IModularFeatures::Get().GetModularFeatureImplementations<IMotionController>(IMotionController::GetModularFeatureName());
+	for (auto MotionController : MotionControllers)
+	{
+		FSteamVRInputDevice* TestSteamVRDevice = static_cast<FSteamVRInputDevice*>(MotionController);
+		if (TestSteamVRDevice != nullptr && !FGenericPlatformMath::IsNaN(TestSteamVRDevice->DeviceSignature) && TestSteamVRDevice->DeviceSignature == 2019)
+		{
+			return TestSteamVRDevice;
+		}
+	}
+	return nullptr;
+}
+
diff --git a/Engine/Plugins/Runtime/Steam/SteamVR/Source/SteamVRInputDevice/Private/SteamVRInputDeviceModule.cpp b/Engine/Plugins/Runtime/Steam/SteamVR/Source/SteamVRInputDevice/Private/SteamVRInputDeviceModule.cpp
new file mode 100644
--- /dev/null
+++ b/Engine/Plugins/Runtime/Steam/SteamVR/Source/SteamVRInputDevice/Private/SteamVRInputDeviceModule.cpp
@@ -0,0 +1,64 @@
+/*
+Copyright 2019 Valve Corporation under https://opensource.org/licenses/BSD-3-Clause
+
+Redistribution and use in source and binary forms, with or without modification,
+are permitted provided that the following conditions are met:
+
+1. Redistributions of source code must retain the above copyright notice, this
+   list of conditions and the following disclaimer.
+
+2. Redistributions in binary form must reproduce the above copyright notice,
+   this list of conditions and the following disclaimer in the documentation
+   and/or other materials provided with the distribution.
+
+3. Neither the name of the copyright holder nor the names of its contributors
+   may be used to endorse or promote products derived from this software
+   without specific prior written permission.
+
+THIS SOFTWARE IS PROVIDED BY THE COPYRIGHT HOLDERS AND CONTRIBUTORS "AS IS" AND
+ANY EXPRESS OR IMPLIED WARRANTIES, INCLUDING, BUT NOT LIMITED TO, THE IMPLIED
+WARRANTIES OF MERCHANTABILITY AND FITNESS FOR A PARTICULAR PURPOSE ARE DISCLAIMED.
+IN NO EVENT SHALL THE COPYRIGHT HOLDER OR CONTRIBUTORS BE LIABLE FOR ANY DIRECT,
+INDIRECT, INCIDENTAL, SPECIAL, EXEMPLARY, OR CONSEQUENTIAL DAMAGES (INCLUDING,
+BUT NOT LIMITED TO, PROCUREMENT OF SUBSTITUTE GOODS OR SERVICES; LOSS OF USE, DATA,
+OR PROFITS; OR BUSINESS INTERRUPTION) HOWEVER CAUSED AND ON ANY THEORY OF LIABILITY,
+WHETHER IN CONTRACT, STRICT LIABILITY, OR TORT (INCLUDING NEGLIGENCE OR OTHERWISE)
+ARISING IN ANY WAY OUT OF THE USE OF THIS SOFTWARE, EVEN IF ADVISED OF THE
+POSSIBILITY OF SUCH DAMAGE.
+*/
+
+
+#include "ISteamVRInputDeviceModule.h"
+#include "IInputDevice.h"
+#include "Runtime/Projects/Public/Interfaces/IPluginManager.h"
+#include "SteamVRInputDevice.h"
+#include "Engine/Engine.h"
+
+class FSteamVRInputDeviceModule : public ISteamVRInputDeviceModule
+{
+	/* Creates a new instance of SteamVR Input Controller **/
+	virtual TSharedPtr<class IInputDevice> CreateInputDevice(const TSharedRef<FGenericApplicationMessageHandler>& InMessageHandler) override
+	{
+		static FName SystemName(TEXT("SteamVR"));
+		if (GEngine->XRSystem.IsValid() && (GEngine->XRSystem->GetSystemName() == SystemName))
+		{
+			return TSharedPtr<class IInputDevice>(new FSteamVRInputDevice(InMessageHandler));
+		}
+		return nullptr;
+	}
+
+	virtual void StartupModule() override;
+	virtual void ShutdownModule() override;
+};
+
+void FSteamVRInputDeviceModule::StartupModule()
+{
+	IModularFeatures::Get().RegisterModularFeature(GetModularFeatureName(), this);
+}
+
+void FSteamVRInputDeviceModule::ShutdownModule()
+{
+	IModularFeatures::Get().UnregisterModularFeature(GetModularFeatureName(), this);
+}
+
+IMPLEMENT_MODULE(FSteamVRInputDeviceModule, SteamVRInputDevice)
diff --git a/Engine/Plugins/Runtime/Steam/SteamVR/Source/SteamVRInputDevice/Private/SteamVRSkeletonDefinition.cpp b/Engine/Plugins/Runtime/Steam/SteamVR/Source/SteamVRInputDevice/Private/SteamVRSkeletonDefinition.cpp
new file mode 100644
--- /dev/null
+++ b/Engine/Plugins/Runtime/Steam/SteamVR/Source/SteamVRInputDevice/Private/SteamVRSkeletonDefinition.cpp
@@ -0,0 +1,161 @@
+/*
+Copyright 2019 Valve Corporation under https://opensource.org/licenses/BSD-3-Clause
+
+Redistribution and use in source and binary forms, with or without modification,
+are permitted provided that the following conditions are met:
+
+1. Redistributions of source code must retain the above copyright notice, this
+   list of conditions and the following disclaimer.
+
+2. Redistributions in binary form must reproduce the above copyright notice,
+   this list of conditions and the following disclaimer in the documentation
+   and/or other materials provided with the distribution.
+
+3. Neither the name of the copyright holder nor the names of its contributors
+   may be used to endorse or promote products derived from this software
+   without specific prior written permission.
+
+THIS SOFTWARE IS PROVIDED BY THE COPYRIGHT HOLDERS AND CONTRIBUTORS "AS IS" AND
+ANY EXPRESS OR IMPLIED WARRANTIES, INCLUDING, BUT NOT LIMITED TO, THE IMPLIED
+WARRANTIES OF MERCHANTABILITY AND FITNESS FOR A PARTICULAR PURPOSE ARE DISCLAIMED.
+IN NO EVENT SHALL THE COPYRIGHT HOLDER OR CONTRIBUTORS BE LIABLE FOR ANY DIRECT,
+INDIRECT, INCIDENTAL, SPECIAL, EXEMPLARY, OR CONSEQUENTIAL DAMAGES (INCLUDING,
+BUT NOT LIMITED TO, PROCUREMENT OF SUBSTITUTE GOODS OR SERVICES; LOSS OF USE, DATA,
+OR PROFITS; OR BUSINESS INTERRUPTION) HOWEVER CAUSED AND ON ANY THEORY OF LIABILITY,
+WHETHER IN CONTRACT, STRICT LIABILITY, OR TORT (INCLUDING NEGLIGENCE OR OTHERWISE)
+ARISING IN ANY WAY OUT OF THE USE OF THIS SOFTWARE, EVEN IF ADVISED OF THE
+POSSIBILITY OF SUCH DAMAGE.
+*/
+
+
+#include "SteamVRSkeletonDefinition.h"
+
+namespace SteamVRSkeleton
+{
+	// Define the index of each bone's parent
+	static const int32 g_BoneParentMap[] = {
+		-1,  // ESteamVRBone_Root = 0,
+		0, 	 // ESteamVRBone_Wrist,
+		1, 	 // ESteamVRBone_Thumb0,
+		2, 	 // ESteamVRBone_Thumb1,
+		3, 	 // ESteamVRBone_Thumb2,
+		4, 	 // ESteamVRBone_Thumb3,
+		1, 	 // ESteamVRBone_IndexFinger0,
+		6, 	 // ESteamVRBone_IndexFinger1,
+		7, 	 // ESteamVRBone_IndexFinger2,
+		8, 	 // ESteamVRBone_IndexFinger3,
+		9, 	 // ESteamVRBone_IndexFinger4,
+		1, 	 // ESteamVRBone_MiddleFinger0,
+		11,  // ESteamVRBone_MiddleFinger1,
+		12,  // ESteamVRBone_MiddleFinger2,
+		13,  // ESteamVRBone_MiddleFinger3,
+		14,  // ESteamVRBone_MiddleFinger4,
+		1, 	 // ESteamVRBone_RingFinger0,
+		16,  // ESteamVRBone_RingFinger1,
+		17,  // ESteamVRBone_RingFinger2,
+		18,  // ESteamVRBone_RingFinger3,
+		19,  // ESteamVRBone_RingFinger4,
+		1, 	 // ESteamVRBone_PinkyFinger0,
+		21,  // ESteamVRBone_PinkyFinger1,
+		22,  // ESteamVRBone_PinkyFinger2,
+		23,  // ESteamVRBone_PinkyFinger3,
+		24,  // ESteamVRBone_PinkyFinger4,
+		0, 	 // ESteamVRBone_Aux_Thumb,
+		0, 	 // ESteamVRBone_Aux_IndexFinger,
+		0, 	 // ESteamVRBone_Aux_MiddleFinger
+		0, 	 // ESteamVRBone_Aux_RingFinger,
+		0	 // ESteamVRBone_Aux_PinkyFinger,
+	};
+
+	// Ensure that the parent bone list matches the size of the list of bones
+	static_assert( sizeof( g_BoneParentMap ) / sizeof(g_BoneParentMap[ 0 ] ) == (size_t)ESteamVRBone_Count, "Bone Parent map is not the right size");
+
+	// Define the names of each bone
+	static const FName g_BoneNames[] = {
+		TEXT("Root"),
+		TEXT("wrist"),
+		TEXT("finger_thumb_0"),
+		TEXT("finger_thumb_1"),
+		TEXT("finger_thumb_2"),
+		TEXT("finger_thumb_end"),
+		TEXT("finger_index_meta"),
+		TEXT("finger_index_0"),
+		TEXT("finger_index_1"),
+		TEXT("finger_index_2"),
+		TEXT("finger_index_end"),
+		TEXT("finger_middle_meta"),
+		TEXT("finger_middle_0"),
+		TEXT("finger_middle_1"),
+		TEXT("finger_middle_2"),
+		TEXT("finger_middle_end"),
+		TEXT("finger_ring_meta"),
+		TEXT("finger_ring_0"),
+		TEXT("finger_ring_1"),
+		TEXT("finger_ring_2"),
+		TEXT("finger_ring_end"),
+		TEXT("finger_pinky_meta"),
+		TEXT("finger_pinky_0"),
+		TEXT("finger_pinky_1"),
+		TEXT("finger_pinky_2"),
+		TEXT("finger_pinky_end"),
+		TEXT("finger_thumb_aux"),
+		TEXT("finger_index_aux"),
+		TEXT("finger_middle_aux"),
+		TEXT("finger_ring_aux"),
+		TEXT("finger_pinky_aux")
+	};
+
+	static_assert( sizeof( g_BoneNames ) / sizeof( g_BoneNames[ 0 ] ) == (size_t)ESteamVRBone_Count, "Bone Name map is not the right size" );
+
+	TArray<TArray<int32>> CreateChildList()
+	{
+		TArray<TArray<int32>> list;
+
+		// Calculate the list of children based on the parent index array, for fast lookup later
+		int32 nBoneCount = (int32)ESteamVRBone_Count;
+		list.SetNum(nBoneCount);
+
+		for (int32 boneIndex = 0; boneIndex < nBoneCount; ++boneIndex)
+		{
+			int32 nParentIndex = g_BoneParentMap[boneIndex];
+			if (nParentIndex != -1)
+			{
+				list[nParentIndex].Add(boneIndex);
+			}
+		}
+
+		return list;
+	}
+
+	static const TArray<TArray<int32>> kBoneChildList = CreateChildList();
+
+
+	const FName& GetBoneName(int32 nBoneIndex)
+	{
+		check(nBoneIndex >= 0 && nBoneIndex < GetBoneCount());
+
+		return g_BoneNames[nBoneIndex];
+	}
+
+	int32 GetParentIndex(int32 nBoneIndex)
+	{
+		check(nBoneIndex >= 0 && nBoneIndex < GetBoneCount());
+
+		return g_BoneParentMap[nBoneIndex];
+	}
+
+	int32 GetChildCount(int32 nBoneIndex)
+	{
+		check(nBoneIndex >= 0 && nBoneIndex < GetBoneCount());
+
+		return kBoneChildList[nBoneIndex].Num();
+	}
+
+	int32 GetChildIndex(int32 nBoneIndex, int32 nChildIndex)
+	{
+		check(nBoneIndex >= 0 && nBoneIndex < GetBoneCount());
+		check(nChildIndex >= 0 && nChildIndex < kBoneChildList[nBoneIndex].Num());
+
+		return kBoneChildList[nBoneIndex][nChildIndex];
+	}
+}
diff --git a/Engine/Plugins/Runtime/Steam/SteamVR/Source/SteamVRInputDevice/Private/SteamVRTrackingRefComponent.cpp b/Engine/Plugins/Runtime/Steam/SteamVR/Source/SteamVRInputDevice/Private/SteamVRTrackingRefComponent.cpp
new file mode 100644
--- /dev/null
+++ b/Engine/Plugins/Runtime/Steam/SteamVR/Source/SteamVRInputDevice/Private/SteamVRTrackingRefComponent.cpp
@@ -0,0 +1,293 @@
+/*
+Copyright 2019 Valve Corporation under https://opensource.org/licenses/BSD-3-Clause
+
+Redistribution and use in source and binary forms, with or without modification,
+are permitted provided that the following conditions are met:
+
+1. Redistributions of source code must retain the above copyright notice, this
+   list of conditions and the following disclaimer.
+
+2. Redistributions in binary form must reproduce the above copyright notice,
+   this list of conditions and the following disclaimer in the documentation
+   and/or other materials provided with the distribution.
+
+3. Neither the name of the copyright holder nor the names of its contributors
+   may be used to endorse or promote products derived from this software
+   without specific prior written permission.
+
+THIS SOFTWARE IS PROVIDED BY THE COPYRIGHT HOLDERS AND CONTRIBUTORS "AS IS" AND
+ANY EXPRESS OR IMPLIED WARRANTIES, INCLUDING, BUT NOT LIMITED TO, THE IMPLIED
+WARRANTIES OF MERCHANTABILITY AND FITNESS FOR A PARTICULAR PURPOSE ARE DISCLAIMED.
+IN NO EVENT SHALL THE COPYRIGHT HOLDER OR CONTRIBUTORS BE LIABLE FOR ANY DIRECT,
+INDIRECT, INCIDENTAL, SPECIAL, EXEMPLARY, OR CONSEQUENTIAL DAMAGES (INCLUDING,
+BUT NOT LIMITED TO, PROCUREMENT OF SUBSTITUTE GOODS OR SERVICES; LOSS OF USE, DATA,
+OR PROFITS; OR BUSINESS INTERRUPTION) HOWEVER CAUSED AND ON ANY THEORY OF LIABILITY,
+WHETHER IN CONTRACT, STRICT LIABILITY, OR TORT (INCLUDING NEGLIGENCE OR OTHERWISE)
+ARISING IN ANY WAY OUT OF THE USE OF THIS SOFTWARE, EVEN IF ADVISED OF THE
+POSSIBILITY OF SUCH DAMAGE.
+*/
+
+#include "SteamVRTrackingRefComponent.h"
+#include "openvr.h"
+#include "SteamVRInput.h"
+
+using namespace vr;
+DEFINE_LOG_CATEGORY_STATIC(LogSteamVRTrackingRefComponent, Log, All);
+
+USteamVRTrackingReferences::USteamVRTrackingReferences()
+{
+	PrimaryComponentTick.bCanEverTick = true;
+}
+
+bool USteamVRTrackingReferences::ShowTrackingReferences(UStaticMesh* TrackingReferenceMesh)
+{
+	if (!TrackingReferenceMesh->IsValidLowLevel())
+	{
+		// TODO: Set default mesh to SteamVR provided render model. Must be backwards compatible to UE4.15
+		UE_LOG(LogSteamVRTrackingRefComponent, Error, TEXT("[TRACKING REFERENCE] No reference mesh found/defined!"));
+		return false;
+	}
+
+	if (VRSystem() && VRCompositor())
+	{
+		// Remove any existing tracking references in-world
+		HideTrackingReferences();
+
+		// Owner must be the player pawn 
+		FVector HMDOffset = GetOwner()->GetActorLocation();
+		// UE_LOG(LogSteamVRTrackingRefComponent, Warning, TEXT("[TRACKING REF] HMD Offset recorded as: %s"), *HMDOffset.ToString());
+
+		// Find all SteamVR Tracking References
+		for (unsigned int id = 0; id < k_unMaxTrackedDeviceCount; ++id)
+		{
+			ETrackedDeviceClass trackedDeviceClass = VRSystem()->GetTrackedDeviceClass(id);
+			
+			if (trackedDeviceClass == ETrackedDeviceClass::TrackedDeviceClass_TrackingReference)
+			{
+				// Extraneous call, used for Debugging
+				//char buf[k_unMaxPropertyStringSize];
+				//uint32 StringBytes = VRSystem()->GetStringTrackedDeviceProperty(id, ETrackedDeviceProperty::Prop_ModelNumber_String, buf, sizeof(buf));
+				//FString stringCache = *FString(UTF8_TO_TCHAR(buf));
+				//UE_LOG(LogSteamVRTrackingRefComponent, Warning, TEXT("[TRACKING REFERENCE] Found the following tracking device: [%i] %s"), id, *stringCache);
+
+				TrackedDevicePose_t TrackedDevicePose = { 0 };
+				VRCompositor()->GetLastPoseForTrackedDeviceIndex(id, &TrackedDevicePose, NULL);
+
+				// Get SteamVR Transform Matrix for this tracking reference
+				HmdMatrix34_t Matrix = TrackedDevicePose.mDeviceToAbsoluteTracking;
+
+				// Transform SteamVR Pose to Unreal Pose
+				FMatrix Pose = FMatrix(
+					FPlane(Matrix.m[0][0], Matrix.m[1][0], Matrix.m[2][0], 0.0f),
+					FPlane(Matrix.m[0][1], Matrix.m[1][1], Matrix.m[2][1], 0.0f),
+					FPlane(Matrix.m[0][2], Matrix.m[1][2], Matrix.m[2][2], 0.0f),
+					FPlane(Matrix.m[0][3], Matrix.m[1][3], Matrix.m[2][3], 1.0f)
+				);
+
+				// Transform SteamVR Rotation Quaternion to a UE FRotator
+				FQuat OrientationQuat;
+				FQuat OrientationPose(Pose);
+				OrientationQuat.X = -OrientationPose.Z;
+				OrientationQuat.Y = OrientationPose.X;
+				OrientationQuat.Z = OrientationPose.Y;
+				OrientationQuat.W = -OrientationPose.W;
+
+				FVector PositionPose = ((FVector(-Pose.M[3][2], Pose.M[3][0], Pose.M[3][1])) * 100.f + HMDOffset);
+				FVector Position = PositionPose;
+
+				//OutOrientation = BaseOrientation.Inverse() * OutOrientation;
+				OrientationQuat.Normalize();
+				FRotator Orientation = OrientationQuat.Rotator();
+				Orientation.Normalize();
+
+				// Create new static mesh component and attach to actor
+				UStaticMeshComponent* TrackingRefStaticMeshComponent;
+				TrackingRefStaticMeshComponent = NewObject<UStaticMeshComponent>(GetOwner());
+				TrackingRefStaticMeshComponent->RegisterComponentWithWorld(GetWorld());
+				TrackingRefStaticMeshComponent->SetSimulatePhysics(false);
+				TrackingRefStaticMeshComponent->SetCollisionEnabled(ECollisionEnabled::NoCollision);
+				TrackingRefStaticMeshComponent->SetWorldTransform(FTransform(Orientation, Position, TrackingReferenceScale));
+				TrackingRefStaticMeshComponent->SetMobility(EComponentMobility::Movable);
+				TrackingRefStaticMeshComponent->AttachToComponent(GetOwner()->GetRootComponent(), FAttachmentTransformRules::KeepWorldTransform);
+
+				//Set Mesh
+				TrackingRefStaticMeshComponent->SetStaticMesh(TrackingReferenceMesh);
+				TrackingRefStaticMeshComponent->SetVisibility(true);
+
+				// Add to list of in-world tracking references 
+				TrackingReferences.Add(TrackingRefStaticMeshComponent);
+			}
+		}
+	}
+
+	return false;
+}
+
+void USteamVRTrackingReferences::HideTrackingReferences()
+{
+	// Check if need to clear up any existing Tracking Reference in-world
+	if (TrackingReferences.Num()> 0)
+	{
+		// Manually destroy each held tracking reference in array
+		for (UStaticMeshComponent* TrackingReference : TrackingReferences)
+		{
+			if (TrackingReference->IsValidLowLevel())
+			{
+				TrackingReference->DestroyComponent();
+			}
+		}
+
+		// Clear up array
+		TrackingReferences.Empty();
+	}
+}
+
+void USteamVRTrackingReferences::BeginPlay()
+{
+	Super::BeginPlay();
+
+	// Reserve memory for tracking device cache upfront to save on performance - we're using 4 as it is the most common setup
+	ActiveTrackingDevices.Reserve(4);
+
+	// Check if this component needs to continually poll for active devices
+	if (ActiveDevicePollFrequency <= 0.f)
+	{
+		this->SetComponentTickEnabled(false);
+	}
+}
+
+void USteamVRTrackingReferences::TickComponent(float DeltaTime, ELevelTick TickType, FActorComponentTickFunction* ThisTickFunction)
+{
+	Super::TickComponent(DeltaTime, TickType, ThisTickFunction);
+
+	// Check if we need to do an active device check
+	CurrentDeltaTime += DeltaTime;
+	if (CurrentDeltaTime >= ActiveDevicePollFrequency)
+	{
+		CurrentDeltaTime = 0.f;
+	}
+	else
+	{
+		return;
+	}
+
+	// Check for any newly activated tracked devices
+	if (VRSystem() && VRInput())
+	{
+		// Find all SteamVR Tracking References
+		for (unsigned int id = 0; id < k_unMaxTrackedDeviceCount; ++id)
+		{
+			ETrackedDeviceClass TrackedDeviceClass = VRSystem()->GetTrackedDeviceClass(id);
+
+			if (TrackedDeviceClass != TrackedDeviceClass_Invalid && 
+				(TrackedDeviceClass == TrackedDeviceClass_Controller || 
+				 TrackedDeviceClass == TrackedDeviceClass_GenericTracker || 
+				 TrackedDeviceClass == TrackedDeviceClass_TrackingReference)
+				)
+			{
+				// Check if this device has already been detected
+				if (FindTrackedDevice(id))
+				{
+					continue;
+				}
+				else
+				{
+					// If not, add this device to the tracked active devices
+					ActiveTrackingDevices.Add(FActiveTrackedDevice(id, false));
+				}
+			
+				// Get Device Class
+				FName DeviceClass = GetDeviceClass(id);
+
+				// Get device model info
+				char buf[k_unMaxPropertyStringSize];
+				uint32 StringBytes = VRSystem()->GetStringTrackedDeviceProperty(id, ETrackedDeviceProperty::Prop_ModelNumber_String, buf, sizeof(buf));
+				FString DeviceModel = *FString(UTF8_TO_TCHAR(buf));
+
+				UE_LOG(LogSteamVRTrackingRefComponent, Warning, TEXT("Found device [%i] %s"), id, *DeviceModel);
+			}
+		}
+
+		// Check if we need to broadcast any activation and deactivation events
+		for (int32 i = 0; i < ActiveTrackingDevices.Num(); i++)
+		{
+			// If a device is flagged as inactive but SteamVR reports it connected, trigger a connected event
+			if (!ActiveTrackingDevices[i].bActivated && VRSystem()->IsTrackedDeviceConnected(ActiveTrackingDevices[i].id))
+			{
+				// Get Device Class
+				FName DeviceClass = GetDeviceClass(ActiveTrackingDevices[i].id);
+
+				// Get device model info
+				char buf[k_unMaxPropertyStringSize];
+				uint32 StringBytes = VRSystem()->GetStringTrackedDeviceProperty(ActiveTrackingDevices[i].id, ETrackedDeviceProperty::Prop_ModelNumber_String, buf, sizeof(buf));
+				FString DeviceModel = *FString(UTF8_TO_TCHAR(buf));
+
+				// Broadcast activated event
+				OnTrackedDeviceActivated.Broadcast(ActiveTrackingDevices[i].id, DeviceClass, DeviceModel);
+
+				// Flag this device as activated
+				ActiveTrackingDevices[i].bActivated = true;
+
+				UE_LOG(LogSteamVRTrackingRefComponent, Warning, TEXT("Device [%i] %s is connected."), ActiveTrackingDevices[i].id, *DeviceModel);
+			}
+
+			// If however a device is flagged as inactive but SteamVR reports it as disconnected, trigger a disconnected event
+			else if (ActiveTrackingDevices[i].bActivated && !VRSystem()->IsTrackedDeviceConnected(ActiveTrackingDevices[i].id))
+			{
+				// Get Device Class
+				FName DeviceClass = GetDeviceClass(ActiveTrackingDevices[i].id);
+
+				// Get device model info
+				char buf[k_unMaxPropertyStringSize];
+				uint32 StringBytes = VRSystem()->GetStringTrackedDeviceProperty(ActiveTrackingDevices[i].id, ETrackedDeviceProperty::Prop_ModelNumber_String, buf, sizeof(buf));
+				FString DeviceModel = *FString(UTF8_TO_TCHAR(buf));
+
+				// Broadcast deactivated event
+				OnTrackedDeviceDeactivated.Broadcast(ActiveTrackingDevices[i].id, DeviceClass, DeviceModel);
+
+				// Flag this device as deactivated
+				ActiveTrackingDevices[i].bActivated = false;
+
+				UE_LOG(LogSteamVRTrackingRefComponent, Warning, TEXT("Device [%i] %s has been disconnected."), ActiveTrackingDevices[i].id, *DeviceModel);
+			}
+		}
+	}
+}
+
+FName USteamVRTrackingReferences::GetDeviceClass(unsigned int id)
+{
+	// Get device class
+	ETrackedDeviceClass TrackedDeviceClass = VRSystem()->GetTrackedDeviceClass(id);
+
+	// Get Device Class
+	FName DeviceClass;
+	switch (TrackedDeviceClass)
+	{
+	case TrackedDeviceClass_Controller:
+		DeviceClass = FName(TEXT("Controller"));
+		break;
+	case TrackedDeviceClass_GenericTracker:
+		DeviceClass = FName(TEXT("GenericTracker"));
+		break;
+	case TrackedDeviceClass_TrackingReference:
+		DeviceClass = FName(TEXT("TrackingReference"));
+		break;
+	default:
+		DeviceClass = FName(TEXT("Unknown"));
+		break;
+	}
+
+	return DeviceClass;
+}
+
+bool USteamVRTrackingReferences::FindTrackedDevice(unsigned int id)
+{
+	for (FActiveTrackedDevice Device : ActiveTrackingDevices)
+	{
+		if (Device.id == id)
+		{
+			return true;
+		}
+	}
+	return false;
+}
diff --git a/Engine/Plugins/Runtime/Steam/SteamVR/Source/SteamVRInputDevice/Private/UE4HandSkeletonDefinition.cpp b/Engine/Plugins/Runtime/Steam/SteamVR/Source/SteamVRInputDevice/Private/UE4HandSkeletonDefinition.cpp
new file mode 100644
--- /dev/null
+++ b/Engine/Plugins/Runtime/Steam/SteamVR/Source/SteamVRInputDevice/Private/UE4HandSkeletonDefinition.cpp
@@ -0,0 +1,131 @@
+/*
+Copyright 2019 Valve Corporation under https://opensource.org/licenses/BSD-3-Clause
+
+Redistribution and use in source and binary forms, with or without modification,
+are permitted provided that the following conditions are met:
+
+1. Redistributions of source code must retain the above copyright notice, this
+   list of conditions and the following disclaimer.
+
+2. Redistributions in binary form must reproduce the above copyright notice,
+   this list of conditions and the following disclaimer in the documentation
+   and/or other materials provided with the distribution.
+
+3. Neither the name of the copyright holder nor the names of its contributors
+   may be used to endorse or promote products derived from this software
+   without specific prior written permission.
+
+THIS SOFTWARE IS PROVIDED BY THE COPYRIGHT HOLDERS AND CONTRIBUTORS "AS IS" AND
+ANY EXPRESS OR IMPLIED WARRANTIES, INCLUDING, BUT NOT LIMITED TO, THE IMPLIED
+WARRANTIES OF MERCHANTABILITY AND FITNESS FOR A PARTICULAR PURPOSE ARE DISCLAIMED.
+IN NO EVENT SHALL THE COPYRIGHT HOLDER OR CONTRIBUTORS BE LIABLE FOR ANY DIRECT,
+INDIRECT, INCIDENTAL, SPECIAL, EXEMPLARY, OR CONSEQUENTIAL DAMAGES (INCLUDING,
+BUT NOT LIMITED TO, PROCUREMENT OF SUBSTITUTE GOODS OR SERVICES; LOSS OF USE, DATA,
+OR PROFITS; OR BUSINESS INTERRUPTION) HOWEVER CAUSED AND ON ANY THEORY OF LIABILITY,
+WHETHER IN CONTRACT, STRICT LIABILITY, OR TORT (INCLUDING NEGLIGENCE OR OTHERWISE)
+ARISING IN ANY WAY OUT OF THE USE OF THIS SOFTWARE, EVEN IF ADVISED OF THE
+POSSIBILITY OF SUCH DAMAGE.
+*/
+
+
+#include "UE4HandSkeletonDefinition.h"
+
+namespace UE4HandSkeleton
+{
+	// Define the index of each bone's parent
+	static const int32 g_BoneParentMap[] = {
+		-1,						// EUE4HandBone_Wrist
+		EUE4HandBone_Wrist, 	// EUE4HandBone_Index_01
+		EUE4HandBone_Index_01, 	// EUE4HandBone_Index_02
+		EUE4HandBone_Index_02, 	// EUE4HandBone_Index_03
+		EUE4HandBone_Wrist, 	// EUE4HandBone_Middle_01
+		EUE4HandBone_Middle_01, // EUE4HandBone_Middle_02
+		EUE4HandBone_Middle_02, // EUE4HandBone_Middle_03
+		EUE4HandBone_Wrist, 	// EUE4HandBone_Pinky_01
+		EUE4HandBone_Pinky_01, 	// EUE4HandBone_Pinky_02
+		EUE4HandBone_Pinky_02, 	// EUE4HandBone_Pinky_03
+		EUE4HandBone_Wrist, 	// EUE4HandBone_Ring_01
+		EUE4HandBone_Ring_01, 	// EUE4HandBone_Ring_02
+		EUE4HandBone_Ring_02,	// EUE4HandBone_Ring_03
+		EUE4HandBone_Wrist,		// EUE4HandBone_Thumb_01
+		EUE4HandBone_Thumb_01,  // EUE4HandBone_Thumb_02
+		EUE4HandBone_Thumb_02,  // EUE4HandBone_Thumb_03
+	};
+
+	// Ensure that the parent bone list matches the size of the list of bones
+	static_assert( sizeof(g_BoneParentMap) / sizeof(g_BoneParentMap[0]) == ( size_t )EUE4HandBone_Count, "Bone Parent map is not the right size" );
+
+	// Define the names of each bone
+	static const FName g_BoneNames[] = {
+		TEXT("hand"),
+		TEXT("index_01"),
+		TEXT("index_02"),
+		TEXT("index_03"),
+		TEXT("middle_01"),
+		TEXT("middle_02"),
+		TEXT("middle_03"),
+		TEXT("pinky_01"),
+		TEXT("pinky_02"),
+		TEXT("pinky_03"),
+		TEXT("ring_01"),
+		TEXT("ring_02"),
+		TEXT("ring_03"),
+		TEXT("thumb_01"),
+		TEXT("thumb_02"),
+		TEXT("thumb_03")
+	};
+
+	static_assert( sizeof(g_BoneNames) / sizeof(g_BoneNames[0]) == ( size_t )EUE4HandBone_Count, "Bone Name map is not the right size" );
+
+	TArray<TArray<int32>> CreateChildList()
+	{
+		TArray<TArray<int32>> list;
+
+		// Calculate the list of children based on the parent index array, for fast lookup later
+		int32 nBoneCount = ( int32 )EUE4HandBone_Count;
+		list.SetNum(nBoneCount);
+
+		for (int32 boneIndex = 0; boneIndex < nBoneCount; ++boneIndex)
+		{
+			int32 nParentIndex = g_BoneParentMap[boneIndex];
+			if (nParentIndex != -1)
+			{
+				list[nParentIndex].Add(boneIndex);
+			}
+		}
+
+		return list;
+	}
+
+	static const TArray<TArray<int32>> kBoneChildList = CreateChildList();
+
+
+	const FName& GetBoneName(int32 nBoneIndex)
+	{
+		check(nBoneIndex >= 0 && nBoneIndex < GetBoneCount());
+
+		return g_BoneNames[nBoneIndex];
+	}
+
+	int32 GetParentIndex(int32 nBoneIndex)
+	{
+		check(nBoneIndex >= 0 && nBoneIndex < GetBoneCount());
+
+		return g_BoneParentMap[nBoneIndex];
+	}
+
+	int32 GetChildCount(int32 nBoneIndex)
+	{
+		check(nBoneIndex >= 0 && nBoneIndex < GetBoneCount());
+
+		return kBoneChildList[nBoneIndex].Num();
+	}
+
+	int32 GetChildIndex(int32 nBoneIndex, int32 nChildIndex)
+	{
+		check(nBoneIndex >= 0 && nBoneIndex < GetBoneCount());
+		check(nChildIndex >= 0 && nChildIndex < kBoneChildList[nBoneIndex].Num());
+
+		return kBoneChildList[nBoneIndex][nChildIndex];
+	}
+}
diff --git a/Engine/Plugins/Runtime/Steam/SteamVR/Source/SteamVRInputDevice/Public/AnimNode_SteamVRInputAnimPose.h b/Engine/Plugins/Runtime/Steam/SteamVR/Source/SteamVRInputDevice/Public/AnimNode_SteamVRInputAnimPose.h
new file mode 100644
--- /dev/null
+++ b/Engine/Plugins/Runtime/Steam/SteamVR/Source/SteamVRInputDevice/Public/AnimNode_SteamVRInputAnimPose.h
@@ -0,0 +1,94 @@
+/*
+Copyright 2019 Valve Corporation under https://opensource.org/licenses/BSD-3-Clause
+
+Redistribution and use in source and binary forms, with or without modification,
+are permitted provided that the following conditions are met:
+
+1. Redistributions of source code must retain the above copyright notice, this
+   list of conditions and the following disclaimer.
+
+2. Redistributions in binary form must reproduce the above copyright notice,
+   this list of conditions and the following disclaimer in the documentation
+   and/or other materials provided with the distribution.
+
+3. Neither the name of the copyright holder nor the names of its contributors
+   may be used to endorse or promote products derived from this software
+   without specific prior written permission.
+
+THIS SOFTWARE IS PROVIDED BY THE COPYRIGHT HOLDERS AND CONTRIBUTORS "AS IS" AND
+ANY EXPRESS OR IMPLIED WARRANTIES, INCLUDING, BUT NOT LIMITED TO, THE IMPLIED
+WARRANTIES OF MERCHANTABILITY AND FITNESS FOR A PARTICULAR PURPOSE ARE DISCLAIMED.
+IN NO EVENT SHALL THE COPYRIGHT HOLDER OR CONTRIBUTORS BE LIABLE FOR ANY DIRECT,
+INDIRECT, INCIDENTAL, SPECIAL, EXEMPLARY, OR CONSEQUENTIAL DAMAGES (INCLUDING,
+BUT NOT LIMITED TO, PROCUREMENT OF SUBSTITUTE GOODS OR SERVICES; LOSS OF USE, DATA,
+OR PROFITS; OR BUSINESS INTERRUPTION) HOWEVER CAUSED AND ON ANY THEORY OF LIABILITY,
+WHETHER IN CONTRACT, STRICT LIABILITY, OR TORT (INCLUDING NEGLIGENCE OR OTHERWISE)
+ARISING IN ANY WAY OUT OF THE USE OF THIS SOFTWARE, EVEN IF ADVISED OF THE
+POSSIBILITY OF SUCH DAMAGE.
+*/
+
+
+#pragma once
+
+#include "CoreMinimal.h"
+#include "Animation/AnimNodeBase.h"
+#include "SteamVRInputDeviceFunctionLibrary.h"
+#include "SteamVRSkeletonDefinition.h"
+#include "AnimNode_SteamVRInputAnimPose.generated.h"
+
+/**
+* Custom animation node to retrieve poses from the Skeletal Input System
+*/
+USTRUCT(BlueprintType)
+struct STEAMVRINPUTDEVICE_API FAnimNode_SteamVRInputAnimPose : public FAnimNode_Base
+{
+	GENERATED_USTRUCT_BODY()
+
+	/** Range of motion for the skeletal input values */
+	UPROPERTY(EditAnywhere, BlueprintReadWrite, Category = Settings, meta = (AlwaysAsPin))
+	EMotionRange MotionRange = EMotionRange::VR_WithoutController;
+
+	/** Which hand should the animation node retrieve skeletal input values for */
+	UPROPERTY(EditAnywhere, BlueprintReadWrite, Category = Settings, meta = (AlwaysAsPin))
+	EHand Hand = EHand::VR_LeftHand;
+
+	/** What kind of skeleton are we dealing with */
+	UPROPERTY(EditAnywhere, BlueprintReadWrite, Category = Settings, meta = (AlwaysAsPin))
+	EHandSkeleton HandSkeleton = EHandSkeleton::VR_SteamVRHandSkeleton;
+
+	/** Should the pose be mirrored so it can be applied to the opposite hand */
+	UPROPERTY(EditAnywhere, BlueprintReadWrite, Category = Settings, meta = ( AlwaysAsPin ))
+	bool Mirror = false;
+
+	/** The UE4 equivalent of the SteamVR Transform values per bone */
+	UPROPERTY(EditAnywhere, BlueprintReadWrite, Category = SteamVRInput)
+	FSteamVRSkeletonTransform SteamVRSkeletalTransform;
+
+	/** SteamVR Skeleton to UE4 retargetting cache */
+	UPROPERTY()
+	FUE4RetargettingRefs UE4RetargettingRefs;
+
+public:
+
+	// FAnimNode_Base interface
+	virtual void Initialize_AnyThread(const FAnimationInitializeContext& Context) override;
+	virtual void CacheBones_AnyThread(const FAnimationCacheBonesContext& Context) override;
+	virtual void Update_AnyThread(const FAnimationUpdateContext & Context) override;
+	virtual void Evaluate_AnyThread(FPoseContext& Output) override;
+	// End of FAnimNode_Base interface
+
+	FAnimNode_SteamVRInputAnimPose();
+
+	/** 
+	 * Retarget the given array of bone transforms for the SteamVR skeleton to the UE4 hand skeleton and apply it to the given FPoseContext. 
+	 * The bone transforms are in the bone's local space.  Assumes that PoseContest.Pose has already been set to its reference pose
+	*/
+	void PoseUE4HandSkeleton(FCompactPose& Pose, const FTransform* BoneTransformsLS, int32 BoneTransformCount);
+
+	/** Retrieve the first active SteamVRInput device present in this game */
+	FSteamVRInputDevice* GetSteamVRInputDevice();
+
+	/** Recursively calculate the model-space transform of the given bone from the local-space transforms on the given pose */
+	FTransform CalcModelSpaceTransform(const FCompactPose& Pose, FCompactPoseBoneIndex BoneIndex);
+
+};
diff --git a/Engine/Plugins/Runtime/Steam/SteamVR/Source/SteamVRInputDevice/Public/AnimNode_SteamVRSetWristTransform.h b/Engine/Plugins/Runtime/Steam/SteamVR/Source/SteamVRInputDevice/Public/AnimNode_SteamVRSetWristTransform.h
new file mode 100644
--- /dev/null
+++ b/Engine/Plugins/Runtime/Steam/SteamVR/Source/SteamVRInputDevice/Public/AnimNode_SteamVRSetWristTransform.h
@@ -0,0 +1,77 @@
+/*
+Copyright 2019 Valve Corporation under https://opensource.org/licenses/BSD-3-Clause
+
+Redistribution and use in source and binary forms, with or without modification,
+are permitted provided that the following conditions are met:
+
+1. Redistributions of source code must retain the above copyright notice, this
+   list of conditions and the following disclaimer.
+
+2. Redistributions in binary form must reproduce the above copyright notice,
+   this list of conditions and the following disclaimer in the documentation
+   and/or other materials provided with the distribution.
+
+3. Neither the name of the copyright holder nor the names of its contributors
+   may be used to endorse or promote products derived from this software
+   without specific prior written permission.
+
+THIS SOFTWARE IS PROVIDED BY THE COPYRIGHT HOLDERS AND CONTRIBUTORS "AS IS" AND
+ANY EXPRESS OR IMPLIED WARRANTIES, INCLUDING, BUT NOT LIMITED TO, THE IMPLIED
+WARRANTIES OF MERCHANTABILITY AND FITNESS FOR A PARTICULAR PURPOSE ARE DISCLAIMED.
+IN NO EVENT SHALL THE COPYRIGHT HOLDER OR CONTRIBUTORS BE LIABLE FOR ANY DIRECT,
+INDIRECT, INCIDENTAL, SPECIAL, EXEMPLARY, OR CONSEQUENTIAL DAMAGES (INCLUDING,
+BUT NOT LIMITED TO, PROCUREMENT OF SUBSTITUTE GOODS OR SERVICES; LOSS OF USE, DATA,
+OR PROFITS; OR BUSINESS INTERRUPTION) HOWEVER CAUSED AND ON ANY THEORY OF LIABILITY,
+WHETHER IN CONTRACT, STRICT LIABILITY, OR TORT (INCLUDING NEGLIGENCE OR OTHERWISE)
+ARISING IN ANY WAY OUT OF THE USE OF THIS SOFTWARE, EVEN IF ADVISED OF THE
+POSSIBILITY OF SUCH DAMAGE.
+*/
+
+
+#pragma once
+
+#include "CoreMinimal.h"
+#include "Animation/AnimNodeBase.h"
+#include "SteamVRInputDeviceFunctionLibrary.h"
+#include "SteamVRSkeletonDefinition.h"
+#include "AnimNode_SteamVRSetWristTransform.generated.h"
+
+/**
+* Custom animation node that sets the wrist transform of a target pose from a reference pose
+*/
+USTRUCT(BlueprintInternalUseOnly)
+struct STEAMVRINPUTDEVICE_API FAnimNode_SteamVRSetWristTransform : public FAnimNode_Base
+{
+	GENERATED_USTRUCT_BODY()
+
+	/** The pose from where we will get the root and/or wrist transform from */
+	UPROPERTY(EditAnywhere, BlueprintReadWrite, Category = Links)
+	FPoseLink ReferencePose;
+
+	/** What kind of skeleton is used in the reference pose */
+	UPROPERTY(EditAnywhere, BlueprintReadWrite, Category = Settings, meta = (AlwaysAsPin))
+	EHandSkeleton HandSkeleton = EHandSkeleton::VR_SteamVRHandSkeleton;
+
+	/** The pose to apply the wrist transform to */
+	UPROPERTY(EditAnywhere, BlueprintReadWrite, Category = Links)
+	FPoseLink TargetPose;
+
+public:
+
+	// FAnimNode_Base interface
+	virtual void Initialize_AnyThread(const FAnimationInitializeContext& Context) override;
+	virtual void CacheBones_AnyThread(const FAnimationCacheBonesContext& Context) override;
+	virtual void Update_AnyThread(const FAnimationUpdateContext& Context) override;
+	virtual void Evaluate_AnyThread(FPoseContext& Output) override;
+	// End of FAnimNode_Base interface
+
+	FAnimNode_SteamVRSetWristTransform();
+
+private:
+	/** The root bone index of the SteamVR & UE4 Skeletons */
+	FCompactPoseBoneIndex RootBoneIndex = FCompactPoseBoneIndex(0);
+
+	/** The wrist bone index of the SteamVR Skeleton */
+	FCompactPoseBoneIndex SteamVRWristBoneIndex = FCompactPoseBoneIndex(1);
+
+};
diff --git a/Engine/Plugins/Runtime/Steam/SteamVR/Source/SteamVRInputDevice/Public/ISteamVRInputDeviceModule.h b/Engine/Plugins/Runtime/Steam/SteamVR/Source/SteamVRInputDevice/Public/ISteamVRInputDeviceModule.h
new file mode 100644
--- /dev/null
+++ b/Engine/Plugins/Runtime/Steam/SteamVR/Source/SteamVRInputDevice/Public/ISteamVRInputDeviceModule.h
@@ -0,0 +1,50 @@
+/*
+Copyright 2019 Valve Corporation under https://opensource.org/licenses/BSD-3-Clause
+
+Redistribution and use in source and binary forms, with or without modification,
+are permitted provided that the following conditions are met:
+
+1. Redistributions of source code must retain the above copyright notice, this
+   list of conditions and the following disclaimer.
+
+2. Redistributions in binary form must reproduce the above copyright notice,
+   this list of conditions and the following disclaimer in the documentation
+   and/or other materials provided with the distribution.
+
+3. Neither the name of the copyright holder nor the names of its contributors
+   may be used to endorse or promote products derived from this software
+   without specific prior written permission.
+
+THIS SOFTWARE IS PROVIDED BY THE COPYRIGHT HOLDERS AND CONTRIBUTORS "AS IS" AND
+ANY EXPRESS OR IMPLIED WARRANTIES, INCLUDING, BUT NOT LIMITED TO, THE IMPLIED
+WARRANTIES OF MERCHANTABILITY AND FITNESS FOR A PARTICULAR PURPOSE ARE DISCLAIMED.
+IN NO EVENT SHALL THE COPYRIGHT HOLDER OR CONTRIBUTORS BE LIABLE FOR ANY DIRECT,
+INDIRECT, INCIDENTAL, SPECIAL, EXEMPLARY, OR CONSEQUENTIAL DAMAGES (INCLUDING,
+BUT NOT LIMITED TO, PROCUREMENT OF SUBSTITUTE GOODS OR SERVICES; LOSS OF USE, DATA,
+OR PROFITS; OR BUSINESS INTERRUPTION) HOWEVER CAUSED AND ON ANY THEORY OF LIABILITY,
+WHETHER IN CONTRACT, STRICT LIABILITY, OR TORT (INCLUDING NEGLIGENCE OR OTHERWISE)
+ARISING IN ANY WAY OUT OF THE USE OF THIS SOFTWARE, EVEN IF ADVISED OF THE
+POSSIBILITY OF SUCH DAMAGE.
+*/
+
+
+#pragma once
+
+#include "CoreMinimal.h"
+#include "Modules/ModuleManager.h"
+#include "IInputDeviceModule.h"
+
+STEAMVRINPUT_API class ISteamVRInputDeviceModule : public IInputDeviceModule
+{
+
+public:
+	static inline ISteamVRInputDeviceModule& Get()
+	{
+		return FModuleManager::LoadModuleChecked<ISteamVRInputDeviceModule>("SteamVRInputDevice");
+	}
+
+	static inline bool IsAvailable()
+	{
+		return FModuleManager::Get().IsModuleLoaded("SteamVRInputDevice");
+	}
+};
diff --git a/Engine/Plugins/Runtime/Steam/SteamVR/Source/SteamVRInputDevice/Public/SteamVRControllerKeys.h b/Engine/Plugins/Runtime/Steam/SteamVR/Source/SteamVRInputDevice/Public/SteamVRControllerKeys.h
new file mode 100644
--- /dev/null
+++ b/Engine/Plugins/Runtime/Steam/SteamVR/Source/SteamVRInputDevice/Public/SteamVRControllerKeys.h
@@ -0,0 +1,153 @@
+/*
+Copyright 2019 Valve Corporation under https://opensource.org/licenses/BSD-3-Clause
+
+Redistribution and use in source and binary forms, with or without modification,
+are permitted provided that the following conditions are met:
+
+1. Redistributions of source code must retain the above copyright notice, this
+   list of conditions and the following disclaimer.
+
+2. Redistributions in binary form must reproduce the above copyright notice,
+   this list of conditions and the following disclaimer in the documentation
+   and/or other materials provided with the distribution.
+
+3. Neither the name of the copyright holder nor the names of its contributors
+   may be used to endorse or promote products derived from this software
+   without specific prior written permission.
+
+THIS SOFTWARE IS PROVIDED BY THE COPYRIGHT HOLDERS AND CONTRIBUTORS "AS IS" AND
+ANY EXPRESS OR IMPLIED WARRANTIES, INCLUDING, BUT NOT LIMITED TO, THE IMPLIED
+WARRANTIES OF MERCHANTABILITY AND FITNESS FOR A PARTICULAR PURPOSE ARE DISCLAIMED.
+IN NO EVENT SHALL THE COPYRIGHT HOLDER OR CONTRIBUTORS BE LIABLE FOR ANY DIRECT,
+INDIRECT, INCIDENTAL, SPECIAL, EXEMPLARY, OR CONSEQUENTIAL DAMAGES (INCLUDING,
+BUT NOT LIMITED TO, PROCUREMENT OF SUBSTITUTE GOODS OR SERVICES; LOSS OF USE, DATA,
+OR PROFITS; OR BUSINESS INTERRUPTION) HOWEVER CAUSED AND ON ANY THEORY OF LIABILITY,
+WHETHER IN CONTRACT, STRICT LIABILITY, OR TORT (INCLUDING NEGLIGENCE OR OTHERWISE)
+ARISING IN ANY WAY OUT OF THE USE OF THIS SOFTWARE, EVEN IF ADVISED OF THE
+POSSIBILITY OF SUCH DAMAGE.
+*/
+
+
+#pragma once
+#include "InputCoreTypes.h"
+
+namespace GenericKeys
+{
+	const FKey SteamVR_MotionController_None("SteamVR_MotionController_None");
+	const FKey SteamVR_HMD_Proximity("SteamVR_HMD_Proximity");
+}
+
+namespace IndexControllerKeys
+{
+	// Thumbstick Directions
+	const FKey ValveIndex_Left_Thumbstick_Up("ValveIndex_Left_Thumbstick_Up");
+	const FKey ValveIndex_Left_Thumbstick_Down("ValveIndex_Left_Thumbstick_Down");
+	const FKey ValveIndex_Left_Thumbstick_Left("ValveIndex_Left_Thumbstick_Left");
+	const FKey ValveIndex_Left_Thumbstick_Right("ValveIndex_Left_Thumbstick_Right");
+	const FKey ValveIndex_Right_Thumbstick_Up("ValveIndex_Right_Thumbstick_Up");
+	const FKey ValveIndex_Right_Thumbstick_Down("ValveIndex_Right_Thumbstick_Down");
+	const FKey ValveIndex_Right_Thumbstick_Left("ValveIndex_Right_Thumbstick_Left");
+	const FKey ValveIndex_Right_Thumbstick_Right("ValveIndex_Right_Thumbstick_Right");
+
+	// Trackpad Directions
+	const FKey ValveIndex_Left_Trackpad_Up("ValveIndex_Left_Trackpad_Up");
+	const FKey ValveIndex_Left_Trackpad_Down("ValveIndex_Left_Trackpad_Down");
+	const FKey ValveIndex_Left_Trackpad_Left("ValveIndex_Left_Trackpad_Left");
+	const FKey ValveIndex_Left_Trackpad_Right("ValveIndex_Left_Trackpad_Right");
+	const FKey ValveIndex_Right_Trackpad_Up("ValveIndex_Right_Trackpad_Up");
+	const FKey ValveIndex_Right_Trackpad_Down("ValveIndex_Right_Trackpad_Down");
+	const FKey ValveIndex_Right_Trackpad_Left("ValveIndex_Right_Trackpad_Left");
+	const FKey ValveIndex_Right_Trackpad_Right("ValveIndex_Right_Trackpad_Right");
+
+	// Special
+	const FKey ValveIndex_Left_Grip_Grab("ValveIndex_Left_Grip_Grab");
+	const FKey ValveIndex_Left_Pinch_Grab("ValveIndex_Left_Pinch_Grab");
+	const FKey ValveIndex_Right_Grip_Grab("ValveIndex_Right_Grip_Grab");
+	const FKey ValveIndex_Right_Pinch_Grab("ValveIndex_Right_Pinch_Grab");
+}
+
+namespace ViveControllerKeys
+{
+	// Trackpad Directions
+	const FKey Vive_Left_Trackpad_Up("Vive_Left_Trackpad_Up");
+	const FKey Vive_Left_Trackpad_Down("Vive_Left_Trackpad_Down");
+	const FKey Vive_Left_Trackpad_Left("Vive_Left_Trackpad_Left");
+	const FKey Vive_Left_Trackpad_Right("Vive_Left_Trackpad_Right");
+	const FKey Vive_Right_Trackpad_Up("Vive_Right_Trackpad_Up");
+	const FKey Vive_Right_Trackpad_Down("Vive_Right_Trackpad_Down");
+	const FKey Vive_Right_Trackpad_Left("Vive_Right_Trackpad_Left");
+	const FKey Vive_Right_Trackpad_Right("Vive_Right_Trackpad_Right");
+}
+
+namespace CosmosKeys
+{
+	// Regular Controller Keys
+	const FKey Cosmos_Left_X_Click("Cosmos_Left_X_Click");
+	const FKey Cosmos_Left_Y_Click("Cosmos_Left_Y_Click");
+	const FKey Cosmos_Left_X_Touch("Cosmos_Left_X_Touch");
+	const FKey Cosmos_Left_Y_Touch("Cosmos_Left_Y_Touch");
+	const FKey Cosmos_Left_Menu_Click("Cosmos_Left_Menu_Click");
+	const FKey Cosmos_Left_Grip_Click("Cosmos_Left_Grip_Click");
+	const FKey Cosmos_Left_Grip_Axis("Cosmos_Left_Grip_Axis");
+	const FKey Cosmos_Left_Trigger_Click("Cosmos_Left_Trigger_Click");
+	const FKey Cosmos_Left_Trigger_Axis("Cosmos_Left_Trigger_Axis");
+	const FKey Cosmos_Left_Trigger_Touch("Cosmos_Left_Trigger_Touch");
+	const FKey Cosmos_Left_Thumbstick_Vector("Cosmos_Left_Thumbstick_Vector");
+	const FKey Cosmos_Left_Thumbstick_X("Cosmos_Left_Thumbstick_X");
+	const FKey Cosmos_Left_Thumbstick_Y("Cosmos_Left_Thumbstick_Y");
+	const FKey Cosmos_Left_Thumbstick_Click("Cosmos_Left_Thumbstick_Click");
+	const FKey Cosmos_Left_Thumbstick_Touch("Cosmos_Left_Thumbstick_Touch");
+	const FKey Cosmos_Left_Bumper_Click("Cosmos_Left_Bumper_Click");
+	const FKey Cosmos_Right_A_Click("Cosmos_Right_A_Click");
+	const FKey Cosmos_Right_B_Click("Cosmos_Right_B_Click");
+	const FKey Cosmos_Right_A_Touch("Cosmos_Right_A_Touch");
+	const FKey Cosmos_Right_B_Touch("Cosmos_Right_B_Touch");
+	const FKey Cosmos_Right_System_Click("Cosmos_Right_System_Click");
+	const FKey Cosmos_Right_Grip_Click("Cosmos_Right_Grip_Click");
+	const FKey Cosmos_Right_Grip_Axis("Cosmos_Right_Grip_Axis");
+	const FKey Cosmos_Right_Trigger_Click("Cosmos_Right_Trigger_Click");
+	const FKey Cosmos_Right_Trigger_Axis("Cosmos_Right_Trigger_Axis");
+	const FKey Cosmos_Right_Trigger_Touch("Cosmos_Right_Trigger_Touch");
+	const FKey Cosmos_Right_Thumbstick_Vector("Cosmos_Right_Thumbstick_Vector");
+	const FKey Cosmos_Right_Thumbstick_X("Cosmos_Right_Thumbstick_X");
+	const FKey Cosmos_Right_Thumbstick_Y("Cosmos_Right_Thumbstick_Y");
+	const FKey Cosmos_Right_Thumbstick_Click("Cosmos_Right_Thumbstick_Click");
+	const FKey Cosmos_Right_Thumbstick_Touch("Cosmos_Right_Thumbstick_Touch");
+	const FKey Cosmos_Right_Bumper_Click("Cosmos_Right_Bumper_Click");
+
+	// Thumbstick Directions
+	const FKey Cosmos_Left_Thumbstick_Up("Cosmos_Left_Thumbstick_Up");
+	const FKey Cosmos_Left_Thumbstick_Down("Cosmos_Left_Thumbstick_Down");
+	const FKey Cosmos_Left_Thumbstick_Left("Cosmos_Left_Thumbstick_Left");
+	const FKey Cosmos_Left_Thumbstick_Right("Cosmos_Left_Thumbstick_Right");
+	const FKey Cosmos_Right_Thumbstick_Up("Cosmos_Right_Thumbstick_Up");
+	const FKey Cosmos_Right_Thumbstick_Down("Cosmos_Right_Thumbstick_Down");
+	const FKey Cosmos_Right_Thumbstick_Left("Cosmos_Right_Thumbstick_Left");
+	const FKey Cosmos_Right_Thumbstick_Right("Cosmos_Right_Thumbstick_Right");
+}
+
+namespace OculusTouchKeys
+{
+	// Thumbstick Directions
+	const FKey OculusTouch_Left_Thumbstick_Up("OculusTouch_Left_Thumbstick_Up");
+	const FKey OculusTouch_Left_Thumbstick_Down("OculusTouch_Left_Thumbstick_Down");
+	const FKey OculusTouch_Left_Thumbstick_Left("OculusTouch_Left_Thumbstick_Left");
+	const FKey OculusTouch_Left_Thumbstick_Right("OculusTouch_Left_Thumbstick_Right");
+	const FKey OculusTouch_Right_Thumbstick_Up("OculusTouch_Right_Thumbstick_Up");
+	const FKey OculusTouch_Right_Thumbstick_Down("OculusTouch_Right_Thumbstick_Down");
+	const FKey OculusTouch_Right_Thumbstick_Left("OculusTouch_Right_Thumbstick_Left");
+	const FKey OculusTouch_Right_Thumbstick_Right("OculusTouch_Right_Thumbstick_Right");
+}
+
+namespace WindowsMRKeys
+{
+	// Trackpad Directions
+	const FKey MixedReality_Left_Trackpad_Up("MixedReality_Left_Trackpad_Up");
+	const FKey MixedReality_Left_Trackpad_Down("MixedReality_Left_Trackpad_Down");
+	const FKey MixedReality_Left_Trackpad_Left("MixedReality_Left_Trackpad_Left");
+	const FKey MixedReality_Left_Trackpad_Right("MixedReality_Left_Trackpad_Right");
+	const FKey MixedReality_Right_Trackpad_Up("MixedReality_Right_Trackpad_Up");
+	const FKey MixedReality_Right_Trackpad_Down("MixedReality_Right_Trackpad_Down");
+	const FKey MixedReality_Right_Trackpad_Left("MixedReality_Right_Trackpad_Left");
+	const FKey MixedReality_Right_Trackpad_Right("MixedReality_Right_Trackpad_Right");
+}
