Deprecate MotionController keys

diff --git a/Engine/Plugins/Runtime/Steam/SteamVR/Source/SteamVRInputDevice/Public/SteamVRInputDeviceFunctionLibrary.h b/Engine/Plugins/Runtime/Steam/SteamVR/Source/SteamVRInputDevice/Public/SteamVRInputDeviceFunctionLibrary.h
new file mode 100644
--- /dev/null
+++ b/Engine/Plugins/Runtime/Steam/SteamVR/Source/SteamVRInputDevice/Public/SteamVRInputDeviceFunctionLibrary.h
@@ -0,0 +1,642 @@
+/*
+Copyright 2019 Valve Corporation under https://opensource.org/licenses/BSD-3-Clause
+
+Redistribution and use in source and binary forms, with or without modification,
+are permitted provided that the following conditions are met:
+
+1. Redistributions of source code must retain the above copyright notice, this
+   list of conditions and the following disclaimer.
+
+2. Redistributions in binary form must reproduce the above copyright notice,
+   this list of conditions and the following disclaimer in the documentation
+   and/or other materials provided with the distribution.
+
+3. Neither the name of the copyright holder nor the names of its contributors
+   may be used to endorse or promote products derived from this software
+   without specific prior written permission.
+
+THIS SOFTWARE IS PROVIDED BY THE COPYRIGHT HOLDERS AND CONTRIBUTORS "AS IS" AND
+ANY EXPRESS OR IMPLIED WARRANTIES, INCLUDING, BUT NOT LIMITED TO, THE IMPLIED
+WARRANTIES OF MERCHANTABILITY AND FITNESS FOR A PARTICULAR PURPOSE ARE DISCLAIMED.
+IN NO EVENT SHALL THE COPYRIGHT HOLDER OR CONTRIBUTORS BE LIABLE FOR ANY DIRECT,
+INDIRECT, INCIDENTAL, SPECIAL, EXEMPLARY, OR CONSEQUENTIAL DAMAGES (INCLUDING,
+BUT NOT LIMITED TO, PROCUREMENT OF SUBSTITUTE GOODS OR SERVICES; LOSS OF USE, DATA,
+OR PROFITS; OR BUSINESS INTERRUPTION) HOWEVER CAUSED AND ON ANY THEORY OF LIABILITY,
+WHETHER IN CONTRACT, STRICT LIABILITY, OR TORT (INCLUDING NEGLIGENCE OR OTHERWISE)
+ARISING IN ANY WAY OUT OF THE USE OF THIS SOFTWARE, EVEN IF ADVISED OF THE
+POSSIBILITY OF SUCH DAMAGE.
+*/
+
+
+#pragma once
+
+#include "CoreMinimal.h"
+#include "SteamVRInputDevice.h"
+#include "Kismet/BlueprintFunctionLibrary.h"
+#include "Features/IModularFeatures.h"
+#include "Runtime/Core/Public/GenericPlatform/GenericPlatformProcess.h"
+#include "IMotionController.h"
+#include "Engine/World.h"
+#include "Kismet/GameplayStatics.h"
+#include "SteamVRInputDeviceFunctionLibrary.generated.h"
+
+#define ACTION_PATH_VIBRATE_LEFT		"/actions/main/out/vibrateleft"
+#define ACTION_PATH_VIBRATE_RIGHT		"/actions/main/out/vibrateright"
+#define MAX_BINDINGINFO_COUNT	5
+
+/** UE4 Bone definition of the SteamVR Skeleton */
+USTRUCT(BlueprintType)
+struct STEAMVRINPUTDEVICE_API FSteamVRSkeletonTransform
+{
+	GENERATED_BODY()
+
+	UPROPERTY(EditAnywhere, BlueprintReadOnly, Category = "SteamVR Input")
+	FTransform Root;
+	UPROPERTY(EditAnywhere, BlueprintReadOnly, Category = "SteamVR Input")
+	FTransform Wrist;
+
+	UPROPERTY(EditAnywhere, BlueprintReadOnly, Category = "SteamVR Input")
+	FTransform Thumb_0;
+	UPROPERTY(EditAnywhere, BlueprintReadOnly, Category = "SteamVR Input")
+	FTransform Thumb_1;
+	UPROPERTY(EditAnywhere, BlueprintReadOnly, Category = "SteamVR Input")
+	FTransform Thumb_2;
+	UPROPERTY(EditAnywhere, BlueprintReadOnly, Category = "SteamVR Input")
+	FTransform Thumb_3;
+	
+	UPROPERTY(EditAnywhere, BlueprintReadOnly, Category = "SteamVR Input")
+	FTransform Index_0;
+	UPROPERTY(EditAnywhere, BlueprintReadOnly, Category = "SteamVR Input")
+	FTransform Index_1;
+	UPROPERTY(EditAnywhere, BlueprintReadOnly, Category = "SteamVR Input")
+	FTransform Index_2;
+	UPROPERTY(EditAnywhere, BlueprintReadOnly, Category = "SteamVR Input")
+	FTransform Index_3;
+	UPROPERTY(EditAnywhere, BlueprintReadOnly, Category = "SteamVR Input")
+	FTransform Index_4;
+
+	UPROPERTY(EditAnywhere, BlueprintReadOnly, Category = "SteamVR Input")
+	FTransform Middle_0;
+	UPROPERTY(EditAnywhere, BlueprintReadOnly, Category = "SteamVR Input")
+	FTransform Middle_1;
+	UPROPERTY(EditAnywhere, BlueprintReadOnly, Category = "SteamVR Input")
+	FTransform Middle_2;
+	UPROPERTY(EditAnywhere, BlueprintReadOnly, Category = "SteamVR Input")
+	FTransform Middle_3;
+	UPROPERTY(EditAnywhere, BlueprintReadOnly, Category = "SteamVR Input")
+	FTransform Middle_4;
+
+	UPROPERTY(EditAnywhere, BlueprintReadOnly, Category = "SteamVR Input")
+	FTransform Ring_0;
+	UPROPERTY(EditAnywhere, BlueprintReadOnly, Category = "SteamVR Input")
+	FTransform Ring_1;
+	UPROPERTY(EditAnywhere, BlueprintReadOnly, Category = "SteamVR Input")
+	FTransform Ring_2;
+	UPROPERTY(EditAnywhere, BlueprintReadOnly, Category = "SteamVR Input")
+	FTransform Ring_3;
+	UPROPERTY(EditAnywhere, BlueprintReadOnly, Category = "SteamVR Input")
+	FTransform Ring_4;
+
+	UPROPERTY(EditAnywhere, BlueprintReadOnly, Category = "SteamVR Input")
+	FTransform Pinky_0;
+	UPROPERTY(EditAnywhere, BlueprintReadOnly, Category = "SteamVR Input")
+	FTransform Pinky_1;
+	UPROPERTY(EditAnywhere, BlueprintReadOnly, Category = "SteamVR Input")
+	FTransform Pinky_2;
+	UPROPERTY(EditAnywhere, BlueprintReadOnly, Category = "SteamVR Input")
+	FTransform Pinky_3;
+	UPROPERTY(EditAnywhere, BlueprintReadOnly, Category = "SteamVR Input")
+	FTransform Pinky_4;
+
+	UPROPERTY(EditAnywhere, BlueprintReadOnly, Category = "SteamVR Input")
+	FTransform Aux_Thumb;
+	UPROPERTY(EditAnywhere, BlueprintReadOnly, Category = "SteamVR Input")
+	FTransform Aux_Index;
+	UPROPERTY(EditAnywhere, BlueprintReadOnly, Category = "SteamVR Input")
+	FTransform Aux_Middle;
+	UPROPERTY(EditAnywhere, BlueprintReadOnly, Category = "SteamVR Input")
+	FTransform Aux_Ring;
+	UPROPERTY(EditAnywhere, BlueprintReadOnly, Category = "SteamVR Input")
+	FTransform Aux_Pinky;
+
+	UPROPERTY(EditAnywhere, BlueprintReadOnly, Category = "SteamVR Input")
+	FTransform Bone_Count;
+};
+
+/** SteamVR finger curls */
+USTRUCT(BlueprintType)
+struct STEAMVRINPUTDEVICE_API FSteamVRFingerCurls
+{
+	GENERATED_BODY()
+
+	UPROPERTY(EditAnywhere, BlueprintReadOnly, Category = "SteamVR Input")
+	float	Thumb;
+
+	UPROPERTY(EditAnywhere, BlueprintReadOnly, Category = "SteamVR Input")
+	float	Index;
+
+	UPROPERTY(EditAnywhere, BlueprintReadOnly, Category = "SteamVR Input")
+	float	Middle;
+	
+	UPROPERTY(EditAnywhere, BlueprintReadOnly, Category = "SteamVR Input")
+	float	Ring;
+
+	UPROPERTY(EditAnywhere, BlueprintReadOnly, Category = "SteamVR Input")
+	float	Pinky;
+};
+
+/** SteamVR finger splays */
+USTRUCT(BlueprintType)
+struct STEAMVRINPUTDEVICE_API FSteamVRFingerSplays
+{
+	GENERATED_BODY()
+
+	UPROPERTY(EditAnywhere, BlueprintReadOnly, Category = "SteamVR Input")
+	float	Thumb_Index;
+
+	UPROPERTY(EditAnywhere, BlueprintReadOnly, Category = "SteamVR Input")
+	float	Index_Middle;
+
+	UPROPERTY(EditAnywhere, BlueprintReadOnly, Category = "SteamVR Input")
+	float	Middle_Ring;
+
+	UPROPERTY(EditAnywhere, BlueprintReadOnly, Category = "SteamVR Input")
+	float	Ring_Pinky;
+};
+
+/** SteamVR actions as defined by the developer */
+USTRUCT(BlueprintType)
+struct STEAMVRINPUTDEVICE_API FSteamVRAction
+{
+	GENERATED_BODY()
+
+	// The SteamVR name of the action (e.g. Teleport, OpenConsole)
+	UPROPERTY(EditAnywhere, BlueprintReadOnly, Category = "SteamVR Input")
+	FName		Name;
+
+	// The path defined for the action (e.g. main/in/{ActionName})
+	UPROPERTY(EditAnywhere, BlueprintReadOnly, Category = "SteamVR Input")
+	FString		Path;
+
+	VRActionHandle_t Handle;				// The handle to the SteamVR Action 
+	VRInputValueHandle_t ActiveOrigin;		// The input value handle of the origin of the latest input event
+
+	FSteamVRAction(const FName ActionName, const FString ActionPath, const VRActionHandle_t ActionHandle, const VRInputValueHandle_t ActionOrigin)
+		: Name(ActionName)
+		, Path(ActionPath)
+		, Handle(ActionHandle)
+		, ActiveOrigin(ActionOrigin)
+	{}
+
+	FSteamVRAction()
+	{}
+};
+
+
+/** SteamVR action set as defined by the developer */
+USTRUCT(BlueprintType)
+struct STEAMVRINPUTDEVICE_API FSteamVRActionSet
+{
+	GENERATED_BODY()
+
+	// The path defined for this action set (e.g. /actions/main)
+	UPROPERTY(EditAnywhere, BlueprintReadOnly, Category = "SteamVR Input")
+	FString		Path;			
+	
+	VRActionSetHandle_t Handle;	// The handle to the SteamVR Action Set
+
+	FSteamVRActionSet(const FString ActionSetPath, const VRActionHandle_t ActionSetHandle)
+		: Path(ActionSetPath)
+		, Handle(ActionSetHandle)
+	{}
+
+	FSteamVRActionSet()
+	{}
+};
+
+/** Information about the tracked device associated from the input source */
+USTRUCT(BlueprintType)
+struct STEAMVRINPUTDEVICE_API FSteamVRInputOriginInfo
+{
+	GENERATED_BODY()
+
+	// The tracked device index for the device or k_unTrackedDeviceInvalid (0xFFFFFFFF)
+	UPROPERTY(EditAnywhere, BlueprintReadOnly, Category = "SteamVR Input")
+	int32	TrackedDeviceIndex;
+	
+	//  The name of the component of the tracked device's render model that represents this input source, or an empty string if there is no associated render model component.
+	UPROPERTY(EditAnywhere, BlueprintReadOnly, Category = "SteamVR Input")
+	FString	RenderModelComponentName;
+
+	//  The tracked device's model info
+	UPROPERTY(EditAnywhere, BlueprintReadOnly, Category = "SteamVR Input")
+	FString	TrackedDeviceModel;
+
+	FSteamVRInputOriginInfo(const int32 InDeviceIndex, const FString InRenderModelComponentName)
+		: TrackedDeviceIndex(InDeviceIndex)
+		, RenderModelComponentName(InRenderModelComponentName)
+	{}
+
+	FSteamVRInputOriginInfo()
+	{
+		TrackedDeviceIndex = 0;
+	}
+};
+
+/** Retargetting information for the SteamVR skeleton to UE4 stock skeleton */
+USTRUCT(BlueprintType)
+struct STEAMVRINPUTDEVICE_API FUE4RetargettingRefs
+{
+	GENERATED_BODY()
+
+	UPROPERTY(BlueprintReadOnly, Category = "SteamVR Input")
+	bool bIsInitialized;
+
+	UPROPERTY(BlueprintReadOnly, Category = "SteamVR Input")
+	bool bIsRightHanded;
+
+	UPROPERTY(BlueprintReadOnly, Category = "SteamVR Input")
+	FVector KnuckleAverageMS_UE4;
+
+	UPROPERTY(BlueprintReadOnly, Category = "SteamVR Input")
+	FVector WristSideDirectionLS_UE4;
+
+	UPROPERTY(BlueprintReadOnly, Category = "SteamVR Input")
+	FVector WristForwardLS_UE4;
+
+	FUE4RetargettingRefs()
+	{
+		bIsInitialized = false;
+		bIsRightHanded = false;
+		KnuckleAverageMS_UE4 = FVector::ZeroVector;
+		WristSideDirectionLS_UE4 = FVector::RightVector;
+		WristForwardLS_UE4 = FVector::ForwardVector;
+	}
+};
+
+/** Information about the input bindings for an action on currently active controller (i.e device path, input path, mode, slot)  */
+USTRUCT(BlueprintType)
+struct STEAMVRINPUTDEVICE_API FSteamVRInputBindingInfo
+{
+	GENERATED_BODY()
+
+	UPROPERTY(BlueprintReadOnly, Category = "SteamVR Input")
+	FName	DevicePathName;
+
+	UPROPERTY(BlueprintReadOnly, Category = "SteamVR Input")
+	FName	InputPathName;
+
+	UPROPERTY(BlueprintReadOnly, Category = "SteamVR Input")
+	FName	ModeName;
+
+	UPROPERTY(BlueprintReadOnly, Category = "SteamVR Input")
+	FName	SlotName;
+
+	FSteamVRInputBindingInfo()
+	{
+		DevicePathName = NAME_None;
+		InputPathName = NAME_None;
+		ModeName = NAME_None;
+		SlotName = NAME_None;
+	}
+
+	FSteamVRInputBindingInfo(FName InDevicePathName, FName InInputPathName, FName InModeName, FName InSlotName) 
+		: DevicePathName(InDevicePathName)
+		, InputPathName(InInputPathName)
+		, ModeName(InModeName)
+		, SlotName(InSlotName)
+	{
+	}
+};
+
+/** Convenience type for SteamVR Hand designation (Left/Right) */
+UENUM(BlueprintType)	
+enum class ESteamVRHand : uint8
+{
+	VR_Left		UMETA(DisplayName = "Left"),
+	VR_Right	UMETA(DisplayName = "Right")
+};
+
+/** Valid range of motion for a skeletal animation */
+UENUM(BlueprintType)
+enum class EMotionRange : uint8
+{
+	VR_WithoutController 	UMETA(DisplayName = "Without Controller"),
+	VR_WithController 		UMETA(DisplayName = "With Controller")
+};
+
+/** Valid values for hands thats used for the Skeletal Input System calls */
+UENUM(BlueprintType)
+enum class EHand : uint8
+{
+	VR_LeftHand 	UMETA(DisplayName = "Left Hand"),
+	VR_RightHand 	UMETA(DisplayName = "Right Hand")
+};
+
+/** Types of known skeletons that this animation node can handle */
+UENUM(BlueprintType)
+enum class EHandSkeleton : uint8
+{
+	VR_SteamVRHandSkeleton 	UMETA(DisplayName = "SteamVR Hand Skeleton"),
+	VR_UE4HandSkeleton 		UMETA(DisplayName = "UE4 Hand Skeleton")
+};
+
+/** Skeletal Tracking Level of a controller */
+UENUM(BlueprintType)
+enum class EControllerFidelity : uint8
+{
+	VR_ControllerFidelity_Estimated 	UMETA(DisplayName = "Controller Fidelity Estimated"),
+	VR_ControllerFidelity_Full 			UMETA(DisplayName = "Controller Fidelity Full"),
+	VR_ControllerFidelity_Partial 		UMETA(DisplayName = "Controller Fidelity Partial")
+};
+
+/** Input String Values for querying user hardware */
+UENUM(BlueprintType)
+enum class ESteamVRInputStringBits : uint8
+{
+	// Which hand the origin is in (e.g. "Left Hand")
+	VR_InputString_Hand					UMETA(DisplayName = "Hand"),
+
+	// What kind of controller the user has in that hand (e.g. "Index Controller")
+	VR_InputString_ControllerType		UMETA(DisplayName = "Controller Type"),
+
+	// What part of that controller is the origin (e.g. "Trackpad")
+	VR_InputString_InputSource			UMETA(DisplayName = "Input Source"),
+
+	// All of the above. (e.g. "Left Hand Index Controller Trackpad")
+	VR_InputString_All					UMETA(DisplayName = "All")
+};
+
+/** Input String Values for querying user hardware */
+UENUM(BlueprintType)
+enum class ESkeletalSummaryDataType : uint8
+{
+	// The data should match the animated transforms in the skeleton transforms. This data will probably be smoothed and may be more latent
+	VR_SummaryType_FromAnimation		UMETA(DisplayName = "From Animation"),
+
+	// The data should be the unprocessed values from the device when available. This data may include more jitter but may be provided with less latency
+	VR_SummaryType_FromDevice			UMETA(DisplayName = "From Device"),
+};
+
+/*
+ * SteamVR Input Extended Functions
+ * Functions and properties defined here are safe for developer use
+ */
+UCLASS()
+class STEAMVRINPUTDEVICE_API USteamVRInputDeviceFunctionLibrary : public UBlueprintFunctionLibrary
+{
+	GENERATED_BODY()
+
+public:
+	/** Retrieve the first available SteamVR Input device currently active in a game */
+	static FSteamVRInputDevice* GetSteamVRInputDevice();
+
+	/**
+	* Get the finger curl and splay for a give hand in the current frame
+	* @param Hand - Which hand to get the finger curls and splay values for
+	* @param FingerCurls - Curl values for each finger pair this frame
+	*/
+	UFUNCTION(BlueprintCallable, Category = "SteamVR Input")
+	static void GetFingerCurlsAndSplays(EHand Hand, FSteamVRFingerCurls& FingerCurls, FSteamVRFingerSplays& FingerSplays, ESkeletalSummaryDataType SummaryDataType = ESkeletalSummaryDataType::VR_SummaryType_FromAnimation);
+
+	/**
+	* Generate haptic feedback in the requested controller
+	* @param Hand - Which hand to send the controller feedback to
+	* @param StartSecondsFromNow - When to start the haptic feedback
+	* @param DurationSeconds - How long to have the haptic feedback active
+	* @param Frequency - Frequency used in the haptic feedback
+	* @param Amplitude - Amplitude used in the haptic feedback
+	*/
+	UFUNCTION(BlueprintCallable, Category="SteamVR Input")
+	static void PlaySteamVR_HapticFeedback(ESteamVRHand Hand, float StartSecondsFromNow, float DurationSeconds = 1.f,
+			float Frequency = 1.f, float Amplitude = 0.5f);
+
+	/**
+	* Check Whether or not Curls and Splay values are being retrieved per frame from the SteamVR Input System
+	* @return LeftHandState - Whether or not curls and splay values are being retrieved from the left hand
+	* @return RightHandState -  Whether or not curls and splay values are being retrieved from the right hand
+	*/
+	UFUNCTION(BlueprintCallable, Category="SteamVR Input")
+	static void GetCurlsAndSplaysState(bool& LeftHandState, bool& RightHandState);
+
+	/**
+	* Check Whether or not controllers attached to either hand have Skeletal Input support
+	* @return LeftHandState - Whether or not the controller attached to the player's left hand have skeletal input support
+	* @return RightHandState -  Whether or not the controller attached to the player's right hand have skeletal input support
+	*/
+	UFUNCTION(BlueprintCallable, Category = "SteamVR Input")
+	static void GetSkeletalState(bool& LeftHandState, bool& RightHandState);
+
+	/** 
+	* Retrieve skeletal tracking level for all controllers 
+	* @return LeftControllerFidelity - The skeletal tracking level of the left controller
+	* @return RightControllerFidelity -  The skeletal tracking level of the right controller
+	*/
+	UFUNCTION(BlueprintCallable, Category = "SteamVR Input")
+	static void GetControllerFidelity(EControllerFidelity& LeftControllerFidelity, EControllerFidelity& RightControllerFidelity);
+
+	/**
+	* Tell SteamVR Whether or not to retrieve Curls and Splay values per frame
+	* @param NewLeftHandState - Whether or not curls and splay values will be retrieved for the left hand
+	* @param NewRightHandState -  Whether or not curls and splay values will be retrieved for the right hand
+	*/
+	UFUNCTION(BlueprintCallable, Category = "SteamVR Input")
+	static void SetCurlsAndSplaysState(bool NewLeftHandState, bool NewRightHandState);
+
+	/**
+	* Check whether we are using a skeleton pose or the raw controller pose for the orientation and position of the motion controller
+	* @param bUseSkeletonPose - Whether or not we are using the skeleton pose instead of the controller raw pose
+	*/
+	UFUNCTION(BlueprintCallable, Category = "SteamVR Input")
+	static void GetPoseSource(bool& bUsingSkeletonPose);
+
+	/**
+	* Set whether to use skeleton pose or the raw controller pose for the orientation and position of the motion controller
+	* @param bUseSkeletonPose - Whether or not to use the skeleton pose instead of the controller raw pose
+	*/
+	UFUNCTION(BlueprintCallable, Category = "SteamVR Input")
+	static void SetPoseSource(bool bUseSkeletonPose);
+
+	/**
+	* Retrieve the live skeletal input bone values from SteamVR
+	* @return LeftHand - Per bone transform values for the left hand skeleton
+	* @return RightHand - Per bone transform values for the right hand skeleton
+	* @param bWithController - Whether or not retrieve skeletal input values with controller
+	* @param bXAxisForward - Whether or not the Skeleton has the X axis facing forward
+	*/
+	UFUNCTION(BlueprintCallable, Category = "SteamVR Input")
+	static void GetSkeletalTransform(FSteamVRSkeletonTransform& LeftHand, FSteamVRSkeletonTransform& RightHand, bool bWithController=false);
+	
+	/**
+	* Retrieve the left hand pose information - position, orientation and velocities
+	* @return Position - Translation from the pose data matrix in UE coordinates
+	* @return Orientation - Orientation derived from the pose data matrix in UE coordinates
+	* @return AngularVelocity - The angular velocity of the hand this frame
+	* @return Velocity - The velocity of the hand this frame
+	*/
+	UFUNCTION(BlueprintCallable, Category = "SteamVR Input")
+	static void GetLeftHandPoseData(FVector& Position, FRotator& Orientation, FVector& AngularVelocity, FVector& Velocity);
+
+	/**
+	* Retrieve the right hand pose information - position, orientation and velocities
+	* @return Position - Translation from the pose data matrix in UE coordinates
+	* @return Orientation - Orientation derived from the pose data matrix in UE coordinates
+	* @return AngularVelocity - The angular velocity of the hand this frame
+	* @return Velocity - The velocity of the hand this frame
+	*/
+	UFUNCTION(BlueprintCallable, Category = "SteamVR Input")
+	static void GetRightHandPoseData(FVector& Position, FRotator& Orientation, FVector& AngularVelocity, FVector& Velocity);
+
+	/**
+	* Retrieve the input actions for this project
+	* @return SteamVRActions - Input actions defined in this project
+	*/
+	UFUNCTION(BlueprintCallable, Category = "SteamVR Input")
+	static void GetSteamVR_ActionArray(TArray<FSteamVRAction>& SteamVRActions);
+
+	/**
+	* Search for a valid action matching the given action name and action set
+	* @param ActionName - The name of the action to look for (e.g. TeleportLeft)
+	* @param ActionSet - The name of the action set that the action belongs to (e.g. main). Default is "main"
+	* @return bresult - The result of the search
+	* @return FoundAction - The action if found
+	* @return FoundActionSet - The action set if found
+	*/
+	UFUNCTION(BlueprintCallable, Category = "SteamVR Input")
+	static void FindSteamVR_Action(FName ActionName, bool& bResult, FSteamVRAction& FoundAction, FSteamVRActionSet& FoundActionSet, FName ActionSet = FName("main"));
+
+	/**
+	* Retrieve the input action sets for this project
+	* @return SteamVRActionSets - Input action sets defined in this project
+	*/
+	UFUNCTION(BlueprintCallable, Category = "SteamVR Input")
+	static void GetSteamVR_ActionSetArray(TArray<FSteamVRActionSet>& SteamVRActionSets);
+
+	/**
+	* Returns information about the tracked device associated from the input source.
+	* @param SteamVRAction - The action that's the source of the input
+	* @return InputOriginInfo - The origin info of the action
+	* @return bool - whether the operation is successful or not
+	*/
+	UFUNCTION(BlueprintCallable, Category = "SteamVR Input")
+	static bool GetSteamVR_OriginTrackedDeviceInfo(FSteamVRAction SteamVRAction, FSteamVRInputOriginInfo& InputOriginInfo);
+
+	/**
+	* Find and return information about the tracked device associated from the input source.
+	* @param SteamVRAction - The action that's the source of the input
+	* @return InputOriginInfo - The origin info of the action
+	* @return bool - whether the operation is successful or not
+	*/
+	UFUNCTION(BlueprintCallable, Category = "SteamVR Input")
+	static void FindSteamVR_OriginTrackedDeviceInfo(FName ActionName, bool& bResult, FSteamVRInputOriginInfo& InputOriginInfo, FName ActionSet = FName("main"));
+
+	/**
+	* Retrieve the localized name of the origin of a given action (e.g. "Left Hand Index Controller Trackpad")
+	* @param SteamVRAction - The action that we will lookup the last active origin for
+	* @param LocalizedParts - Bitfields to specify which origin parts to return
+	* @return OriginLocalizedName -  The localized name of the origin of a given action (e.g. "Left Hand Index Controller Trackpad")
+	*/
+	UFUNCTION(BlueprintCallable, Category = "SteamVR Input")
+	static void GetSteamVR_OriginLocalizedName(FSteamVRAction SteamVRAction, TArray<ESteamVRInputStringBits> LocalizedParts, FString& OriginLocalizedName);
+
+	/**
+	* Show the current binding of a given action in the user's HMD
+	* @param SteamVRAction - The action that we will lookup the current binding for
+	* @param SteamVRActionSet - The action set that the action belongs to
+	*/
+	UFUNCTION(BlueprintCallable, Category = "SteamVR Input")
+	static void ShowSteamVR_ActionOrigin(FSteamVRAction SteamVRAction, FSteamVRActionSet SteamVRActionSet);
+
+	/**
+	* Search and show the current binding of a provided action name and action set in the user's HMD
+	* @param SteamVRAction - The action that we will lookup the current binding for
+	* @param SteamVRActionSet - The action set that the action belongs to. Defaults to "main"
+	*/
+	UFUNCTION(BlueprintCallable, Category = "SteamVR Input")
+	static bool FindSteamVR_ActionOrigin(FName ActionName, FName ActionSet = FName("main"));
+
+	/**
+	* Returns the data for the hand transform at any point in time from current time, given a relative number of seconds
+	* @param Hand - The hand that we're going to retrieve the transform for
+	* @return Position - The position of the hand at the point in time, given a relative number of seconds, from the current time
+	* @return Orientation - The rotation of the hand at the point in time, given a relative number of seconds, from the current time
+	* @return bool - Whether or not the call was succesful
+	*/
+	UFUNCTION(BlueprintCallable, Category = "SteamVR Input")
+	static bool GetSteamVR_HandPoseRelativeToNow(FVector& Position, FRotator& Orientation, ESteamVRHand Hand = ESteamVRHand::VR_Left, float PredictedSecondsFromNow = 0.f);
+
+	/**
+	* Returns the the current value of the global PredictedSecondsFromNow use in any Get Pose Action Data calls (i.e. Getting controller transform)
+	* A value of -9999.f triggers a GetPoseActionDataForNextFrame, otherwise GetPoseActionRelativeToNow is called with this value
+	* @return float - The current Predicted Seconds From Now from the SteamVRInput device
+	*/
+	UFUNCTION(BlueprintCallable, Category = "SteamVR Input")
+	static float GetSteamVR_GlobalPredictedSecondsFromNow();
+
+	/**
+	* Sets the the current value of the global PredictedSecondsFromNow to use in any Get Pose Action Data calls (i.e. Getting controller transform)
+	* A value of -9999.f will trigger a GetPoseActionDataForNextFrame, otherwise GetPoseActionRelativeToNow will be called with this value
+	* @param NewValue - The value for PredictedSecondsFromNow that will be used by the SteamVRInput device for Get Action Pose Data calls 
+	* @return float - The current Predicted Seconds From Now from the SteamVRInput device
+	*/
+	UFUNCTION(BlueprintCallable, Category = "SteamVR Input")
+	static float SetSteamVR_GlobalPredictedSecondsFromNow(float NewValue);
+
+	/**
+	* Shows all current bindings for the current controller in the user's headset
+	*/
+	UFUNCTION(BlueprintCallable, Category = "SteamVR Input")
+	static void ShowAllSteamVR_ActionOrigins();
+
+	/**
+	* Retrieves useful information about the SteamVR input bindings for an action.
+	* @param SteamVRActionHandle - The action handle of the action that binding info will be retrieved for the currently active controller. Use Find SteamVRAction node to get a handle
+	* @return SteamVRInputBindingInfo - Array of binding info for an action with the currently active controller
+	*/
+	UFUNCTION(BlueprintCallable, Category = "SteamVR Input")
+	static TArray<FSteamVRInputBindingInfo> GetSteamVR_InputBindingInfo(FSteamVRAction SteamVRActionHandle);
+
+	/**
+	* Retrieves useful information about the SteamVR input bindings with a given action name and action set.
+	* @param ActionName - The name of the action that binding info will be retrieved for the currently active controller
+	* @param ActionSet - The name of the action set that the action belongs in
+	* @return SteamVRInputBindingInfo - Array of binding info for an action with the currently active controller
+	*/
+	UFUNCTION(BlueprintCallable, Category = "SteamVR Input")
+	static TArray<FSteamVRInputBindingInfo> FindSteamVR_InputBindingInfo(FName ActionName, FName ActionSet = FName("main"));
+
+	/**
+	* Sets the zero pose for the seated tracker coordinate system to the current position and yaw of the HMD. 
+	* After this call, calls that pass TrackingUniverseSeated as the origin will be relative to this new zero pose.
+	*
+	* NOTE: This function overrides the user's previously saved seated zero pose and should only be called as the result of a user action.
+	* Users are also able to set their seated zero pose via the SteamVR Dashboard.
+	*/
+	UFUNCTION(BlueprintCallable, Category = "SteamVR Input")
+	static bool ResetSeatedPosition();
+
+	/**
+	* Returns the user's HMD's current IPD (interpupillary distance) setting in millimetres.
+	* @return float - The current IPD setting of the user's headset in millimetres
+	*/
+	UFUNCTION(BlueprintCallable, Category = "SteamVR Input")
+	static float GetUserIPD();
+
+	/**
+	* Get the SteamVR Bone Transform value in UE coordinates
+	* @param SteamBoneTransform - The SteamVR Bone Transform value to get the UE coordinates for
+	*/
+	static FTransform GetUETransform(VRBoneTransform_t SteamBoneTransform);
+
+	/** Regenerate the action manifest used by SteamVR Input System */
+	static void RegenActionManifest();
+	
+	/** Regenerate Controller Bindings for supported SteamVR Controller types */
+	static void RegenControllerBindings();
+	
+	/** Live reload the action manifest and register it to the SteamVR Input System */
+	static void ReloadActionManifest();
+	
+	/** Open the SteamVR Controller Input Dashboard in the user#s default browser */
+	static void LaunchBindingsURL();
+};
diff --git a/Engine/Plugins/Runtime/Steam/SteamVR/Source/SteamVRInputDevice/Public/SteamVRInputPublic.h b/Engine/Plugins/Runtime/Steam/SteamVR/Source/SteamVRInputDevice/Public/SteamVRInputPublic.h
new file mode 100644
--- /dev/null
+++ b/Engine/Plugins/Runtime/Steam/SteamVR/Source/SteamVRInputDevice/Public/SteamVRInputPublic.h
@@ -0,0 +1,69 @@
+/*
+Copyright 2019 Valve Corporation under https://opensource.org/licenses/BSD-3-Clause
+
+Redistribution and use in source and binary forms, with or without modification,
+are permitted provided that the following conditions are met:
+
+1. Redistributions of source code must retain the above copyright notice, this
+   list of conditions and the following disclaimer.
+
+2. Redistributions in binary form must reproduce the above copyright notice,
+   this list of conditions and the following disclaimer in the documentation
+   and/or other materials provided with the distribution.
+
+3. Neither the name of the copyright holder nor the names of its contributors
+   may be used to endorse or promote products derived from this software
+   without specific prior written permission.
+
+THIS SOFTWARE IS PROVIDED BY THE COPYRIGHT HOLDERS AND CONTRIBUTORS "AS IS" AND
+ANY EXPRESS OR IMPLIED WARRANTIES, INCLUDING, BUT NOT LIMITED TO, THE IMPLIED
+WARRANTIES OF MERCHANTABILITY AND FITNESS FOR A PARTICULAR PURPOSE ARE DISCLAIMED.
+IN NO EVENT SHALL THE COPYRIGHT HOLDER OR CONTRIBUTORS BE LIABLE FOR ANY DIRECT,
+INDIRECT, INCIDENTAL, SPECIAL, EXEMPLARY, OR CONSEQUENTIAL DAMAGES (INCLUDING,
+BUT NOT LIMITED TO, PROCUREMENT OF SUBSTITUTE GOODS OR SERVICES; LOSS OF USE, DATA,
+OR PROFITS; OR BUSINESS INTERRUPTION) HOWEVER CAUSED AND ON ANY THEORY OF LIABILITY,
+WHETHER IN CONTRACT, STRICT LIABILITY, OR TORT (INCLUDING NEGLIGENCE OR OTHERWISE)
+ARISING IN ANY WAY OUT OF THE USE OF THIS SOFTWARE, EVEN IF ADVISED OF THE
+POSSIBILITY OF SUCH DAMAGE.
+*/
+
+
+struct FSteamVRSkeletonTransform_t
+{
+	FTransform Wrist;
+
+	FTransform Thumb_Metacarpal;
+	FTransform Thumb_Proximal;
+	FTransform Thumb_Middle;
+	FTransform Thumb_Distal;
+	FTransform Thumb_Tip;
+	FTransform Thumb_Aux;
+
+	FTransform Index_Metacarpal;
+	FTransform Index_Proximal;
+	FTransform Index_Middle;
+	FTransform Index_Distal;
+	FTransform Index_Tip;
+	FTransform Index_Aux;
+
+	FTransform Middle_Metacarpal;
+	FTransform Middle_Proximal;
+	FTransform Middle_Middle;
+	FTransform Middle_Distal;
+	FTransform Middle_Tip;
+	FTransform Middle_Aux;
+
+	FTransform Ring_Metacarpal;
+	FTransform Ring_Proximal;
+	FTransform Ring_Middle;
+	FTransform Ring_Distal;
+	FTransform Ring_Tip;
+	FTransform Ring_Aux;
+
+	FTransform Pinky_Metacarpal;
+	FTransform Pinky_Proximal;
+	FTransform Pinky_Middle;
+	FTransform Pinky_Distal;
+	FTransform Pinky_Tip;
+	FTransform Pinky_Aux;
+};
diff --git a/Engine/Plugins/Runtime/Steam/SteamVR/Source/SteamVRInputDevice/Public/SteamVRSkeletonDefinition.h b/Engine/Plugins/Runtime/Steam/SteamVR/Source/SteamVRInputDevice/Public/SteamVRSkeletonDefinition.h
new file mode 100644
--- /dev/null
+++ b/Engine/Plugins/Runtime/Steam/SteamVR/Source/SteamVRInputDevice/Public/SteamVRSkeletonDefinition.h
@@ -0,0 +1,95 @@
+/*
+Copyright 2019 Valve Corporation under https://opensource.org/licenses/BSD-3-Clause
+
+Redistribution and use in source and binary forms, with or without modification,
+are permitted provided that the following conditions are met:
+
+1. Redistributions of source code must retain the above copyright notice, this
+   list of conditions and the following disclaimer.
+
+2. Redistributions in binary form must reproduce the above copyright notice,
+   this list of conditions and the following disclaimer in the documentation
+   and/or other materials provided with the distribution.
+
+3. Neither the name of the copyright holder nor the names of its contributors
+   may be used to endorse or promote products derived from this software
+   without specific prior written permission.
+
+THIS SOFTWARE IS PROVIDED BY THE COPYRIGHT HOLDERS AND CONTRIBUTORS "AS IS" AND
+ANY EXPRESS OR IMPLIED WARRANTIES, INCLUDING, BUT NOT LIMITED TO, THE IMPLIED
+WARRANTIES OF MERCHANTABILITY AND FITNESS FOR A PARTICULAR PURPOSE ARE DISCLAIMED.
+IN NO EVENT SHALL THE COPYRIGHT HOLDER OR CONTRIBUTORS BE LIABLE FOR ANY DIRECT,
+INDIRECT, INCIDENTAL, SPECIAL, EXEMPLARY, OR CONSEQUENTIAL DAMAGES (INCLUDING,
+BUT NOT LIMITED TO, PROCUREMENT OF SUBSTITUTE GOODS OR SERVICES; LOSS OF USE, DATA,
+OR PROFITS; OR BUSINESS INTERRUPTION) HOWEVER CAUSED AND ON ANY THEORY OF LIABILITY,
+WHETHER IN CONTRACT, STRICT LIABILITY, OR TORT (INCLUDING NEGLIGENCE OR OTHERWISE)
+ARISING IN ANY WAY OUT OF THE USE OF THIS SOFTWARE, EVEN IF ADVISED OF THE
+POSSIBILITY OF SUCH DAMAGE.
+*/
+
+
+#pragma once
+
+#include "CoreMinimal.h"
+
+/**
+ * Enum for each bone in the SteamVR hand skeleton
+*/
+enum ESteamVRBone : int8
+{
+	ESteamVRBone_Root = 0,
+	ESteamVRBone_Wrist,
+	ESteamVRBone_Thumb0,
+	ESteamVRBone_Thumb1,
+	ESteamVRBone_Thumb2,
+	ESteamVRBone_Thumb3,
+	ESteamVRBone_IndexFinger0,
+	ESteamVRBone_IndexFinger1,
+	ESteamVRBone_IndexFinger2,
+	ESteamVRBone_IndexFinger3,
+	ESteamVRBone_IndexFinger4,
+	ESteamVRBone_MiddleFinger0,
+	ESteamVRBone_MiddleFinger1,
+	ESteamVRBone_MiddleFinger2,
+	ESteamVRBone_MiddleFinger3,
+	ESteamVRBone_MiddleFinger4,
+	ESteamVRBone_RingFinger0,
+	ESteamVRBone_RingFinger1,
+	ESteamVRBone_RingFinger2,
+	ESteamVRBone_RingFinger3,
+	ESteamVRBone_RingFinger4,
+	ESteamVRBone_PinkyFinger0,
+	ESteamVRBone_PinkyFinger1,
+	ESteamVRBone_PinkyFinger2,
+	ESteamVRBone_PinkyFinger3,
+	ESteamVRBone_PinkyFinger4,
+	ESteamVRBone_Aux_Thumb,
+	ESteamVRBone_Aux_IndexFinger,
+	ESteamVRBone_Aux_MiddleFinger,
+	ESteamVRBone_Aux_RingFinger,
+	ESteamVRBone_Aux_PinkyFinger,
+	ESteamVRBone_Count
+};
+
+
+/**
+ * Convenience functions to access static data as defined 
+ * by the SteamVR Skeletal Input skeleton
+*/
+namespace SteamVRSkeleton
+{
+	/** Returns the number of bones in the skeleton */
+	inline int32	GetBoneCount() { return ESteamVRBone_Count; }
+
+	/** Returns the name of the bone at the given index */
+	const FName&	GetBoneName(int32 nBoneIndex);
+
+	/** Returns the index of the parent bone of the given bone.  Returns -1 if the bone does not have a parent */
+	int32			GetParentIndex(int32 nBoneIndex);
+
+	/** Returns the number of children of the given bone */
+	int32			GetChildCount(int32 nBoneIndex);
+
+	/** Returns the index of the nth child of the given bone */
+	int32			GetChildIndex(int32 nBoneIndex, int32 nChildIndex);
+};
diff --git a/Engine/Plugins/Runtime/Steam/SteamVR/Source/SteamVRInputDevice/Public/SteamVRTrackingRefComponent.h b/Engine/Plugins/Runtime/Steam/SteamVR/Source/SteamVRInputDevice/Public/SteamVRTrackingRefComponent.h
new file mode 100644
--- /dev/null
+++ b/Engine/Plugins/Runtime/Steam/SteamVR/Source/SteamVRInputDevice/Public/SteamVRTrackingRefComponent.h
@@ -0,0 +1,104 @@
+/*
+Copyright 2019 Valve Corporation under https://opensource.org/licenses/BSD-3-Clause
+
+Redistribution and use in source and binary forms, with or without modification,
+are permitted provided that the following conditions are met:
+
+1. Redistributions of source code must retain the above copyright notice, this
+   list of conditions and the following disclaimer.
+
+2. Redistributions in binary form must reproduce the above copyright notice,
+   this list of conditions and the following disclaimer in the documentation
+   and/or other materials provided with the distribution.
+
+3. Neither the name of the copyright holder nor the names of its contributors
+   may be used to endorse or promote products derived from this software
+   without specific prior written permission.
+
+THIS SOFTWARE IS PROVIDED BY THE COPYRIGHT HOLDERS AND CONTRIBUTORS "AS IS" AND
+ANY EXPRESS OR IMPLIED WARRANTIES, INCLUDING, BUT NOT LIMITED TO, THE IMPLIED
+WARRANTIES OF MERCHANTABILITY AND FITNESS FOR A PARTICULAR PURPOSE ARE DISCLAIMED.
+IN NO EVENT SHALL THE COPYRIGHT HOLDER OR CONTRIBUTORS BE LIABLE FOR ANY DIRECT,
+INDIRECT, INCIDENTAL, SPECIAL, EXEMPLARY, OR CONSEQUENTIAL DAMAGES (INCLUDING,
+BUT NOT LIMITED TO, PROCUREMENT OF SUBSTITUTE GOODS OR SERVICES; LOSS OF USE, DATA,
+OR PROFITS; OR BUSINESS INTERRUPTION) HOWEVER CAUSED AND ON ANY THEORY OF LIABILITY,
+WHETHER IN CONTRACT, STRICT LIABILITY, OR TORT (INCLUDING NEGLIGENCE OR OTHERWISE)
+ARISING IN ANY WAY OUT OF THE USE OF THIS SOFTWARE, EVEN IF ADVISED OF THE
+POSSIBILITY OF SUCH DAMAGE.
+*/
+
+#pragma once
+
+#include "Runtime/Engine/Classes/Engine/StaticMesh.h"
+#include "Components/ActorComponent.h"
+#include "Components/StaticMeshComponent.h"
+#include "SteamVRTrackingRefComponent.generated.h"
+
+// Delegates
+DECLARE_DYNAMIC_MULTICAST_DELEGATE_ThreeParams(FComponentTrackingActivatedSignature, int32, DeviceID, FName, DeviceClass, FString, DeviceModel);
+DECLARE_DYNAMIC_MULTICAST_DELEGATE_ThreeParams(FComponentTrackingDeactivatedSignature, int32, DeviceID, FName, DeviceClass, FString, DeviceModel);
+
+UCLASS( ClassGroup=(Custom), meta=(BlueprintSpawnableComponent) )
+class STEAMVRINPUTDEVICE_API USteamVRTrackingReferences : public UActorComponent
+{
+	GENERATED_BODY()
+
+public:	
+	USteamVRTrackingReferences();
+
+	/** Blueprint event - When a new active device is recognized */
+	UPROPERTY(BlueprintAssignable, Category = "VR")
+	FComponentTrackingActivatedSignature OnTrackedDeviceActivated;
+
+	/** When an active device gets deactivated */
+	UPROPERTY(BlueprintAssignable, Category = "VR")
+	FComponentTrackingDeactivatedSignature OnTrackedDeviceDeactivated;
+
+	// TODO: Set default mesh to SteamVR provided render model. Must be backwards compatible to UE4.15
+
+	/** Display Tracking References in-world */
+	UFUNCTION(BlueprintCallable, Category = "SteamVR Input")
+	bool ShowTrackingReferences(UStaticMesh* TrackingReferenceMesh);
+
+	/** Remove Tracking References in-world */
+	UFUNCTION(BlueprintCallable, Category = "SteamVR Input")
+	void HideTrackingReferences();
+
+	/** Scale to apply to the tracking reference mesh */
+	UPROPERTY(EditAnywhere, BlueprintReadWrite, Category = "SteamVR Input")
+	float ActiveDevicePollFrequency = 1.f;
+
+	/** Scale to apply to the tracking reference mesh */
+	UPROPERTY(EditAnywhere, BlueprintReadWrite, Category = "SteamVR Input")
+	FVector TrackingReferenceScale = FVector(1.f);
+
+	/** Currently displayed Tracking References in-world */
+	UPROPERTY(BlueprintReadOnly, Category = "SteamVR Input")
+	TArray<UStaticMeshComponent*> TrackingReferences;
+
+private:
+	/** Represents a tracked device with a flag on whether its activated or not */
+	struct FActiveTrackedDevice
+	{
+		unsigned int		id;		// The SteamVR id of this device
+		bool		bActivated;		// Whether or not this device has been activated
+
+		FActiveTrackedDevice(unsigned int inId, bool inActivated)
+			: id(inId)
+			, bActivated(inActivated)
+		{}
+
+	};
+
+	/** Cache for current delta time */
+	float CurrentDeltaTime = 0.f;
+
+	/** Cache to hold tracked devices registered in SteamVR */
+	TArray<FActiveTrackedDevice> ActiveTrackingDevices;
+
+protected:
+	virtual void BeginPlay() override;	
+	virtual void TickComponent(float DeltaTime, ELevelTick TickType, FActorComponentTickFunction* ThisTickFunction) override;
+	FName GetDeviceClass(unsigned int id);
+	bool FindTrackedDevice(unsigned int id);
+};
diff --git a/Engine/Plugins/Runtime/Steam/SteamVR/Source/SteamVRInputDevice/SteamVRInputDevice.Build.cs b/Engine/Plugins/Runtime/Steam/SteamVR/Source/SteamVRInputDevice/SteamVRInputDevice.Build.cs
new file mode 100644
--- /dev/null
+++ b/Engine/Plugins/Runtime/Steam/SteamVR/Source/SteamVRInputDevice/SteamVRInputDevice.Build.cs
@@ -0,0 +1,57 @@
+using UnrealBuildTool;
+using System.IO;
+
+public class SteamVRInputDevice : ModuleRules
+{
+	public SteamVRInputDevice(ReadOnlyTargetRules Target) : base(Target)
+    {
+        PCHUsage = PCHUsageMode.UseExplicitOrSharedPCHs;
+        PrivatePCHHeaderFile = "Public/ISteamVRInputDeviceModule.h";
+
+        PublicIncludePaths.Add(Path.Combine(ModuleDirectory, "Public"));
+        PrivateIncludePaths.Add(Path.Combine(ModuleDirectory, "Private"));
+
+        PrivateIncludePathModuleNames.AddRange(new string[]
+         {
+            "TargetPlatform",
+            "SteamVR"
+        });
+
+        PrivateDependencyModuleNames.AddRange(
+			new string[]
+			{
+				"Core",
+                "CoreUObject",
+                "ApplicationCore",
+                "Engine",
+                "HeadMountedDisplay",
+                "SteamVR",
+                "Json",
+                "JsonUtilities"
+			}
+			);
+
+        PublicDependencyModuleNames.AddRange(
+            new string[]
+            {
+                "SteamVR",
+                "SteamVRInput",
+                "InputCore",
+                "InputDevice"
+            }
+            );
+
+        if (Target.bBuildEditor == true)
+        {
+            DynamicallyLoadedModuleNames.Add("UnrealEd");
+        }
+
+        AddEngineThirdPartyPrivateStaticDependencies(Target, "OpenVR");
+
+        if (Target.Platform == UnrealTargetPlatform.Win64 || Target.Platform == UnrealTargetPlatform.Win32 || (Target.Platform == UnrealTargetPlatform.Linux && Target.Architecture.StartsWith("x86_64")))
+        {
+            AddEngineThirdPartyPrivateStaticDependencies(Target, "OpenGL");
+            PrivateDependencyModuleNames.Add("OpenGLDrv");
+        }
+    }
+}
diff --git a/Engine/Plugins/Runtime/Steam/SteamVR/SteamVR.uplugin b/Engine/Plugins/Runtime/Steam/SteamVR/SteamVR.uplugin
--- a/Engine/Plugins/Runtime/Steam/SteamVR/SteamVR.uplugin
+++ b/Engine/Plugins/Runtime/Steam/SteamVR/SteamVR.uplugin
@@ -33,14 +33,27 @@
 			]
 		},
 		{
-			"Name": "SteamVRController",
+			"Name": "SteamVRInput",
 			"Type": "Runtime",
+			"LoadingPhase": "PostConfigInit",
+			"WhitelistPlatforms": [
+				"Win64"
+			]
+		},
+		{
+			"Name": "SteamVRInputDevice",
+			"Type": "Runtime",
+			"LoadingPhase": "PostConfigInit",
+			"WhitelistPlatforms": [
+				"Win64"
+			]
+		},
+		{
+			"Name": "SteamVREditor",
+			"Type": "Editor",
 			"LoadingPhase": "Default",
 			"WhitelistPlatforms": [
-				"Win64",
-				"Win32",
-				"Linux",
-				"Mac"
+				"Win64"
 			]
 		}
 	],
diff --git a/Engine/Plugins/Runtime/WindowsMixedReality/Source/WindowsMixedRealitySpatialInput/Private/WindowsMixedRealitySpatialInput.cpp b/Engine/Plugins/Runtime/WindowsMixedReality/Source/WindowsMixedRealitySpatialInput/Private/WindowsMixedRealitySpatialInput.cpp
--- a/Engine/Plugins/Runtime/WindowsMixedReality/Source/WindowsMixedRealitySpatialInput/Private/WindowsMixedRealitySpatialInput.cpp
+++ b/Engine/Plugins/Runtime/WindowsMixedReality/Source/WindowsMixedRealitySpatialInput/Private/WindowsMixedRealitySpatialInput.cpp
@@ -169,63 +169,78 @@ namespace WindowsMixedReality
 	void FWindowsMixedRealitySpatialInput::SendAxisEvents(uint32 source)
 	{
 		FKey key;
-		float position;
+		FVector position;
 
 		for (int i = 0; i < 2; i++)
 		{
 			HMDHand hand = (HMDHand)i;
 
 			// Trigger
-			position = FWindowsMixedRealityStatics::GetAxisPosition(hand, HMDInputControllerAxes::SelectValue);
+			position.X = FWindowsMixedRealityStatics::GetAxisPosition(hand, HMDInputControllerAxes::SelectValue);
 			key = (hand == HMDHand::Left) ?
 				EKeys::MotionController_Left_TriggerAxis :
 				EKeys::MotionController_Right_TriggerAxis;
 
-			SendControllerAxisEvent(MessageHandler, source, key, position);
+			SendControllerAxisEvent(MessageHandler, source, key, position.X);
+			key = (hand == HMDHand::Left) ?
+				EKeys::MixedReality_Left_Trigger_Axis :
+				EKeys::MixedReality_Right_Trigger_Axis;
+
+			SendControllerAxisEvent(MessageHandler, source, key, position.X);
 
 			// Thumbstick X
-			position = FWindowsMixedRealityStatics::GetAxisPosition(hand, HMDInputControllerAxes::ThumbstickX);
+			position.X = FWindowsMixedRealityStatics::GetAxisPosition(hand, HMDInputControllerAxes::ThumbstickX);
 			key = (hand == HMDHand::Left) ?
 				EKeys::MotionController_Left_Thumbstick_X :
 				EKeys::MotionController_Right_Thumbstick_X;
 
-			SendControllerAxisEvent(MessageHandler, source, key, position);
+			SendControllerAxisEvent(MessageHandler, source, key, position.X);
+			key = (hand == HMDHand::Left) ?
+				EKeys::MixedReality_Left_Thumbstick_X :
+				EKeys::MixedReality_Right_Thumbstick_X;
+
+			SendControllerAxisEvent(MessageHandler, source, key, position.X);
 
 			// Thumbstick Y
-			position = FWindowsMixedRealityStatics::GetAxisPosition(hand, HMDInputControllerAxes::ThumbstickY);
+			position.Y = FWindowsMixedRealityStatics::GetAxisPosition(hand, HMDInputControllerAxes::ThumbstickY);
 			key = (hand == HMDHand::Left) ?
 				EKeys::MotionController_Left_Thumbstick_Y :
 				EKeys::MotionController_Right_Thumbstick_Y;
 
-			SendControllerAxisEvent(MessageHandler, source, key, position);
+			SendControllerAxisEvent(MessageHandler, source, key, position.Y);
+			key = (hand == HMDHand::Left) ?
+				EKeys::MixedReality_Left_Thumbstick_Y :
+				EKeys::MixedReality_Right_Thumbstick_Y;
+
+			SendControllerAxisEvent(MessageHandler, source, key, position.Y);
 
 			// Touchpad X
-			position = FWindowsMixedRealityStatics::GetAxisPosition(hand, HMDInputControllerAxes::TouchpadX);
+			position.X = FWindowsMixedRealityStatics::GetAxisPosition(hand, HMDInputControllerAxes::TouchpadX);
 			key = (hand == HMDHand::Left) ?
-				FSpatialInputKeys::LeftTouchpadX :
-				FSpatialInputKeys::RightTouchpadX;
+				EKeys::MixedReality_Left_Trackpad_X :
+				EKeys::MixedReality_Right_Trackpad_X;
 
-			if ((key == FSpatialInputKeys::LeftTouchpadX && !isLeftTouchpadTouched) ||
-				(key == FSpatialInputKeys::RightTouchpadX && !isRightTouchpadTouched))
+			if ((key == EKeys::MixedReality_Left_Trackpad_X && !isLeftTouchpadTouched) ||
+				(key == EKeys::MixedReality_Right_Trackpad_X && !isRightTouchpadTouched))
 			{
-				position = 0.0f;
+				position.X = 0.0f;
 			}
 
-			SendControllerAxisEvent(MessageHandler, source, key, position);
+			SendControllerAxisEvent(MessageHandler, source, key, position.X);
 
 			// Touchpad Y
-			position = FWindowsMixedRealityStatics::GetAxisPosition(hand, HMDInputControllerAxes::TouchpadY);
+			position.Y = FWindowsMixedRealityStatics::GetAxisPosition(hand, HMDInputControllerAxes::TouchpadY);
 			key = (hand == HMDHand::Left) ?
-				FSpatialInputKeys::LeftTouchpadY :
-				FSpatialInputKeys::RightTouchpadY;
+				EKeys::MixedReality_Left_Trackpad_Y :
+				EKeys::MixedReality_Right_Trackpad_Y;
 
-			if ((key == FSpatialInputKeys::LeftTouchpadY && !isLeftTouchpadTouched) ||
-				(key == FSpatialInputKeys::RightTouchpadY && !isRightTouchpadTouched))
+			if ((key == EKeys::MixedReality_Left_Trackpad_Y && !isLeftTouchpadTouched) ||
+				(key == EKeys::MixedReality_Right_Trackpad_Y && !isRightTouchpadTouched))
 			{
-				position = 0.0f;
+				position.Y = 0.0f;
 			}
 
-			SendControllerAxisEvent(MessageHandler, source, key, position);
+			SendControllerAxisEvent(MessageHandler, source, key, position.Y);
 		}
 	}
 
@@ -247,6 +262,11 @@ namespace WindowsMixedReality
 					EKeys::MotionController_Right_Trigger;
 
 				SendControllerButtonEvent(MessageHandler, source, key, pressState);
+				key = (hand == HMDHand::Left) ?
+					EKeys::MixedReality_Left_Trigger_Click :
+					EKeys::MixedReality_Right_Trigger_Click;
+
+				SendControllerButtonEvent(MessageHandler, source, key, pressState);
 			}
 
 			// Grasp
@@ -258,6 +278,11 @@ namespace WindowsMixedReality
 					EKeys::MotionController_Right_Grip1;
 
 				SendControllerButtonEvent(MessageHandler, source, key, pressState);
+				key = (hand == HMDHand::Left) ?
+					EKeys::MixedReality_Left_Grip_Click :
+					EKeys::MixedReality_Right_Grip_Click;
+
+				SendControllerButtonEvent(MessageHandler, source, key, pressState);
 			}
 
 			// Menu
@@ -265,8 +290,8 @@ namespace WindowsMixedReality
 			if (pressState != HMDInputPressState::NotApplicable)
 			{
 				key = (hand == HMDHand::Left) ?
-					FSpatialInputKeys::LeftMenu :
-					FSpatialInputKeys::RightMenu;
+					EKeys::MixedReality_Left_Menu_Click :
+					EKeys::MixedReality_Right_Menu_Click;
 
 				SendControllerButtonEvent(MessageHandler, source, key, pressState);
 			}
@@ -280,6 +305,11 @@ namespace WindowsMixedReality
 					EKeys::MotionController_Right_Thumbstick;
 
 				SendControllerButtonEvent(MessageHandler, source, key, pressState);
+				key = (hand == HMDHand::Left) ?
+					EKeys::MixedReality_Left_Thumbstick_Click :
+					EKeys::MixedReality_Right_Thumbstick_Click;
+
+				SendControllerButtonEvent(MessageHandler, source, key, pressState);
 			}
 
 			// Touchpad press
@@ -287,8 +317,8 @@ namespace WindowsMixedReality
 			if (pressState != HMDInputPressState::NotApplicable)
 			{
 				key = (hand == HMDHand::Left) ?
-					FSpatialInputKeys::LeftTouchpadPress :
-					FSpatialInputKeys::RightTouchpadPress;
+					EKeys::MixedReality_Left_Trackpad_Click :
+					EKeys::MixedReality_Right_Trackpad_Click;
 
 				SendControllerButtonEvent(MessageHandler, source, key, pressState);
 			}
@@ -298,14 +328,14 @@ namespace WindowsMixedReality
 			if (pressState != HMDInputPressState::NotApplicable)
 			{
 				key = (hand == HMDHand::Left) ?
-					FSpatialInputKeys::LeftTouchpadIsTouched :
-					FSpatialInputKeys::RightTouchpadIsTouched;
+					EKeys::MixedReality_Left_Trackpad_Touch :
+					EKeys::MixedReality_Right_Trackpad_Touch;
 
-				if (key == FSpatialInputKeys::LeftTouchpadIsTouched)
+				if (key == EKeys::MixedReality_Left_Trackpad_Touch)
 				{
 					isLeftTouchpadTouched = pressState == HMDInputPressState::Pressed;
 				}
-				else if (key == FSpatialInputKeys::RightTouchpadIsTouched)
+				else if (key == EKeys::MixedReality_Right_Trackpad_Touch)
 				{
 					isRightTouchpadTouched = pressState == HMDInputPressState::Pressed;
 				}
@@ -695,66 +725,6 @@ namespace WindowsMixedReality
 
 	void FWindowsMixedRealitySpatialInput::RegisterKeys() noexcept
 	{
-		EKeys::AddMenuCategoryDisplayInfo(
-			WindowsMixedRealityCategoryName,
-			LOCTEXT(WindowsMixedRealityCategory, WindowsMixedRealityCategoryFriendlyName),
-			TEXT("GraphEditor.PadEvent_16x"));
-
-		EKeys::AddKey(FKeyDetails(
-			FSpatialInputKeys::LeftMenu,
-			LOCTEXT(LeftMenuName, LeftMenuFriendlyName),
-			FKeyDetails::GamepadKey,
-			WindowsMixedRealityCategoryName));
-		EKeys::AddKey(FKeyDetails(
-			FSpatialInputKeys::RightMenu,
-			LOCTEXT(RightMenuName, RightMenuFriendlyName),
-			FKeyDetails::GamepadKey,
-			WindowsMixedRealityCategoryName));
-
-		EKeys::AddKey(FKeyDetails(
-			FSpatialInputKeys::LeftTouchpadPress,
-			LOCTEXT(LeftTouchpadPressName, LeftTouchpadPressFriendlyName),
-			FKeyDetails::GamepadKey,
-			WindowsMixedRealityCategoryName));
-		EKeys::AddKey(FKeyDetails(
-			FSpatialInputKeys::RightTouchpadPress,
-			LOCTEXT(RightTouchpadPressName, RightTouchpadPressFriendlyName),
-			FKeyDetails::GamepadKey,
-			WindowsMixedRealityCategoryName));
-
-		EKeys::AddKey(FKeyDetails(
-			FSpatialInputKeys::LeftTouchpadIsTouched,
-			LOCTEXT(LeftTouchpadIsTouchedName, LeftTouchpadIsTouchedFriendlyName),
-			FKeyDetails::GamepadKey,
-			WindowsMixedRealityCategoryName));
-		EKeys::AddKey(FKeyDetails(
-			FSpatialInputKeys::RightTouchpadIsTouched,
-			LOCTEXT(RightTouchpadIsTouchedName, RightTouchpadIsTouchedFriendlyName),
-			FKeyDetails::GamepadKey,
-			WindowsMixedRealityCategoryName));
-
-		EKeys::AddKey(FKeyDetails(
-			FSpatialInputKeys::LeftTouchpadX,
-			LOCTEXT(LeftTouchpadXName, LeftTouchpadXFriendlyName),
-			FKeyDetails::GamepadKey | FKeyDetails::FloatAxis,
-			WindowsMixedRealityCategoryName));
-		EKeys::AddKey(FKeyDetails(
-			FSpatialInputKeys::RightTouchpadX,
-			LOCTEXT(RightTouchpadXName, RightTouchpadXFriendlyName),
-			FKeyDetails::GamepadKey | FKeyDetails::FloatAxis,
-			WindowsMixedRealityCategoryName));
-
-		EKeys::AddKey(FKeyDetails(
-			FSpatialInputKeys::LeftTouchpadY,
-			LOCTEXT(LeftTouchpadYName, LeftTouchpadYFriendlyName),
-			FKeyDetails::GamepadKey | FKeyDetails::FloatAxis,
-			WindowsMixedRealityCategoryName));
-		EKeys::AddKey(FKeyDetails(
-			FSpatialInputKeys::RightTouchpadY,
-			LOCTEXT(RightTouchpadYName, RightTouchpadYFriendlyName),
-			FKeyDetails::GamepadKey | FKeyDetails::FloatAxis,
-			WindowsMixedRealityCategoryName));
-			
 		EKeys::AddKey(FKeyDetails(FSpatialInputKeys::TapGesture, LOCTEXT(TapGestureName, "Windows Spatial Input Tap Gesture"), FKeyDetails::GamepadKey));
 		EKeys::AddKey(FKeyDetails(FSpatialInputKeys::DoubleTapGesture, LOCTEXT(DoubleTapGestureName, "Windows Spatial Input Double Tap Gesture"), FKeyDetails::GamepadKey));
 		EKeys::AddKey(FKeyDetails(FSpatialInputKeys::HoldGesture, LOCTEXT(HoldGestureName, "Windows Spatial Input Hold Gesture"), FKeyDetails::GamepadKey));
diff --git a/Engine/Plugins/Runtime/WindowsMixedReality/Source/WindowsMixedRealitySpatialInput/Private/WindowsMixedRealitySpatialInputConstants.cpp b/Engine/Plugins/Runtime/WindowsMixedReality/Source/WindowsMixedRealitySpatialInput/Private/WindowsMixedRealitySpatialInputConstants.cpp
--- a/Engine/Plugins/Runtime/WindowsMixedReality/Source/WindowsMixedRealitySpatialInput/Private/WindowsMixedRealitySpatialInputConstants.cpp
+++ b/Engine/Plugins/Runtime/WindowsMixedReality/Source/WindowsMixedRealitySpatialInput/Private/WindowsMixedRealitySpatialInputConstants.cpp
@@ -2,21 +2,6 @@
 
 #include "WindowsSpatialInputDefinitions.h"
 
-const FKey FSpatialInputKeys::LeftMenu(LeftMenuName);
-const FKey FSpatialInputKeys::RightMenu(RightMenuName);
-
-const FKey FSpatialInputKeys::LeftTouchpadPress(LeftTouchpadPressName);
-const FKey FSpatialInputKeys::RightTouchpadPress(RightTouchpadPressName);
-
-const FKey FSpatialInputKeys::LeftTouchpadIsTouched(LeftTouchpadIsTouchedName);
-const FKey FSpatialInputKeys::RightTouchpadIsTouched(RightTouchpadIsTouchedName);
-
-const FKey FSpatialInputKeys::LeftTouchpadX(LeftTouchpadXName);
-const FKey FSpatialInputKeys::RightTouchpadX(RightTouchpadXName);
-
-const FKey FSpatialInputKeys::LeftTouchpadY(LeftTouchpadYName);
-const FKey FSpatialInputKeys::RightTouchpadY(RightTouchpadYName);
-
 const FKey FSpatialInputKeys::TapGesture(TapGestureName);
 const FKey FSpatialInputKeys::DoubleTapGesture(DoubleTapGestureName);
 const FKey FSpatialInputKeys::HoldGesture(HoldGestureName);
diff --git a/Engine/Plugins/Runtime/WindowsMixedReality/Source/WindowsMixedRealitySpatialInput/Private/WindowsSpatialInputDefinitions.h b/Engine/Plugins/Runtime/WindowsMixedReality/Source/WindowsMixedRealitySpatialInput/Private/WindowsSpatialInputDefinitions.h
--- a/Engine/Plugins/Runtime/WindowsMixedReality/Source/WindowsMixedRealitySpatialInput/Private/WindowsSpatialInputDefinitions.h
+++ b/Engine/Plugins/Runtime/WindowsMixedReality/Source/WindowsMixedRealitySpatialInput/Private/WindowsSpatialInputDefinitions.h
@@ -4,38 +4,10 @@
 
 #include "InputCoreTypes.h"
 
-/*
-* Controller button name definitions
-*/
-#define LeftMenuName "WindowsSpatialInput_LeftMenu"
-#define LeftMenuFriendlyName "Windows Spatial Input (L) Menu"
-#define RightMenuName "WindowsSpatialInput_RightMenu"
-#define RightMenuFriendlyName "Windows Spatial Input (R) Menu"
-
-#define LeftTouchpadPressName "WindowsSpatialInput_LeftTouchpad"
-#define LeftTouchpadPressFriendlyName "Windows Spatial Input (L) Touchpad"
-#define RightTouchpadPressName "WindowsSpatialInput_RightTouchpad"
-#define RightTouchpadPressFriendlyName "Windows Spatial Input (R) Touchpad"
-
-#define LeftTouchpadIsTouchedName "WindowsSpatialInput_LeftTouchpadIsTouched"
-#define LeftTouchpadIsTouchedFriendlyName "Windows Spatial Input (L) Touchpad Is Touched"
-#define RightTouchpadIsTouchedName "WindowsSpatialInput_RightTouchpadIsTouched"	
-#define RightTouchpadIsTouchedFriendlyName "Windows Spatial Input (R) Touchpad Is Touched"
-
 /*
 * Controller axis name definitions
 */
 
-#define LeftTouchpadXName "WindowsSpatialInput_LeftTouchpad_X"
-#define LeftTouchpadXFriendlyName "Windows Spatial Input (L) Touchpad X"
-#define RightTouchpadXName "WindowsSpatialInput_RightTouchpad_X"
-#define RightTouchpadXFriendlyName "Windows Spatial Input (R) Touchpad X"
-
-#define LeftTouchpadYName "WindowsSpatialInput_LeftTouchpad_Y"
-#define LeftTouchpadYFriendlyName "Windows Spatial Input (L) Touchpad Y"
-#define RightTouchpadYName "WindowsSpatialInput_RightTouchpad_Y"
-#define RightTouchpadYFriendlyName "Windows Spatial Input (R) Touchpad Y"
-
 #define TapGestureName "WindowsSpatialInput_TapGesture"	
 #define DoubleTapGestureName "WindowsSpatialInput_DoubleTapGesture"	
 #define HoldGestureName "WindowsSpatialInput_HoldGesture"	
@@ -75,24 +47,6 @@
 
 struct FSpatialInputKeys
 {
-	static const FKey LeftGrasp;
-	static const FKey RightGrasp;
-
-	static const FKey LeftMenu;
-	static const FKey RightMenu;
-
-	static const FKey LeftTouchpadPress;
-	static const FKey RightTouchpadPress;
-
-	static const FKey LeftTouchpadIsTouched;
-	static const FKey RightTouchpadIsTouched;
-
-	static const FKey LeftTouchpadX;
-	static const FKey RightTouchpadX;
-
-	static const FKey LeftTouchpadY;
-	static const FKey RightTouchpadY;
-	
 	static const FKey TapGesture;
 	static const FKey DoubleTapGesture;
 
diff --git a/Engine/Source/Editor/UnrealEd/Private/SKeySelector.cpp b/Engine/Source/Editor/UnrealEd/Private/SKeySelector.cpp
--- a/Engine/Source/Editor/UnrealEd/Private/SKeySelector.cpp
+++ b/Engine/Source/Editor/UnrealEd/Private/SKeySelector.cpp
@@ -94,7 +94,7 @@ void SKeySelector::Construct(const FArguments& InArgs)
 
 	for (FKey Key : AllKeys)
 	{
-		if (!InArgs._FilterBlueprintBindable || Key.IsBindableInBlueprints())
+		if (Key.IsBindableToActions() && (!InArgs._FilterBlueprintBindable || Key.IsBindableInBlueprints()))
 		{
 			const FName KeyMenuCategory = Key.GetMenuCategory();
 			FKeyTreeItem* KeyCategory = TreeRootsForCatgories.Find(KeyMenuCategory);
diff --git a/Engine/Source/Runtime/InputCore/Classes/InputCoreTypes.h b/Engine/Source/Runtime/InputCore/Classes/InputCoreTypes.h
--- a/Engine/Source/Runtime/InputCore/Classes/InputCoreTypes.h
+++ b/Engine/Source/Runtime/InputCore/Classes/InputCoreTypes.h
@@ -70,6 +70,7 @@ struct INPUTCORE_API FKey
 	bool IsVectorAxis() const;
 	bool IsBindableInBlueprints() const;
 	bool ShouldUpdateAxisWithoutSamples() const;
+	bool IsBindableToActions() const;
 	FText GetDisplayName(bool bLongDisplayName = true) const;
 	FString ToString() const;
 	FName GetFName() const;
@@ -129,12 +130,13 @@ struct INPUTCORE_API FKeyDetails
 		FloatAxis				 = 1 << 5,
 		VectorAxis				 = 1 << 6,
 		UpdateAxisWithoutSamples = 1 << 7,
+		NotActionBindableKey	 = 1 << 8,
 
 		NoFlags                 = 0,
 	};
 
-	FKeyDetails(const FKey InKey, const TAttribute<FText>& InLongDisplayName, const uint8 InKeyFlags = 0, const FName InMenuCategory = NAME_None, const TAttribute<FText>& InShortDisplayName = TAttribute<FText>() );
-	FKeyDetails(const FKey InKey, const TAttribute<FText>& InLongDisplayName, const TAttribute<FText>& InShortDisplayName, const uint8 InKeyFlags = 0, const FName InMenuCategory = NAME_None);
+	FKeyDetails(const FKey InKey, const TAttribute<FText>& InLongDisplayName, const uint16 InKeyFlags = 0, const FName InMenuCategory = NAME_None, const TAttribute<FText>& InShortDisplayName = TAttribute<FText>() );
+	FKeyDetails(const FKey InKey, const TAttribute<FText>& InLongDisplayName, const TAttribute<FText>& InShortDisplayName, const uint16 InKeyFlags = 0, const FName InMenuCategory = NAME_None);
 
 	FORCEINLINE bool IsModifierKey() const { return bIsModifierKey != 0; }
 	FORCEINLINE bool IsGamepadKey() const { return bIsGamepadKey != 0; }
@@ -144,13 +146,14 @@ struct INPUTCORE_API FKeyDetails
 	FORCEINLINE bool IsVectorAxis() const { return AxisType == EInputAxisType::Vector; }
 	FORCEINLINE bool IsBindableInBlueprints() const { return bIsBindableInBlueprints != 0; }
 	FORCEINLINE bool ShouldUpdateAxisWithoutSamples() const { return bShouldUpdateAxisWithoutSamples != 0; }
+	FORCEINLINE bool IsBindableToActions() const { return bIsBindableToActions != 0; }
 	FORCEINLINE FName GetMenuCategory() const { return MenuCategory; }
 	FText GetDisplayName(const bool bLongDisplayName = true) const;
 	FORCEINLINE const FKey& GetKey() const { return Key; }
 
 private:
 
-	void CommonInit(const uint8 InKeyFlags);	
+	void CommonInit(const uint16 InKeyFlags);	
 
 	enum class EInputAxisType : uint8
 	{
@@ -169,6 +172,7 @@ private:
 	uint8 bIsMouseButton : 1;
 	uint8 bIsBindableInBlueprints : 1;
 	uint8 bShouldUpdateAxisWithoutSamples : 1;
+	uint8 bIsBindableToActions : 1;
 	EInputAxisType AxisType;
 
 	TAttribute<FText> LongDisplayName;
@@ -460,6 +464,7 @@ struct INPUTCORE_API EKeys
 	 
 	//   Motion Controller Axes
 	//		Left Controller
+	static const FKey MotionController_Left_Thumbstick_XY;
 	static const FKey MotionController_Left_Thumbstick_X;
 	static const FKey MotionController_Left_Thumbstick_Y;
 	static const FKey MotionController_Left_TriggerAxis;
@@ -467,6 +472,7 @@ struct INPUTCORE_API EKeys
 	static const FKey MotionController_Left_Grip2Axis;
 
 	//		Right Controller
+	static const FKey MotionController_Right_Thumbstick_XY;
 	static const FKey MotionController_Right_Thumbstick_X;
 	static const FKey MotionController_Right_Thumbstick_Y;
 	static const FKey MotionController_Right_TriggerAxis;
@@ -497,6 +503,152 @@ struct INPUTCORE_API EKeys
 	static const FKey Android_Volume_Down;
 	static const FKey Android_Menu;
 
+	// Google Daydream
+	static const FKey Daydream_Left_Select_Click;
+	static const FKey Daydream_Left_Trackpad_X;
+	static const FKey Daydream_Left_Trackpad_Y;
+	static const FKey Daydream_Left_Trackpad_Click;
+	static const FKey Daydream_Left_Trackpad_Touch;
+	static const FKey Daydream_Right_Select_Click;
+	static const FKey Daydream_Right_Trackpad_X;
+	static const FKey Daydream_Right_Trackpad_Y;
+	static const FKey Daydream_Right_Trackpad_Click;
+	static const FKey Daydream_Right_Trackpad_Touch;
+
+	// HTC Vive Controller
+	static const FKey Vive_Left_System_Click;
+	static const FKey Vive_Left_Grip_Click;
+	static const FKey Vive_Left_Menu_Click;
+	static const FKey Vive_Left_Trigger_Click;
+	static const FKey Vive_Left_Trigger_Axis;
+	static const FKey Vive_Left_Trackpad_X;
+	static const FKey Vive_Left_Trackpad_Y;
+	static const FKey Vive_Left_Trackpad_Click;
+	static const FKey Vive_Left_Trackpad_Touch;
+	static const FKey Vive_Right_System_Click;
+	static const FKey Vive_Right_Grip_Click;
+	static const FKey Vive_Right_Menu_Click;
+	static const FKey Vive_Right_Trigger_Click;
+	static const FKey Vive_Right_Trigger_Axis;
+	static const FKey Vive_Right_Trackpad_X;
+	static const FKey Vive_Right_Trackpad_Y;
+	static const FKey Vive_Right_Trackpad_Click;
+	static const FKey Vive_Right_Trackpad_Touch;
+
+	// Microsoft Mixed Reality Motion Controller
+	static const FKey MixedReality_Left_Menu_Click;
+	static const FKey MixedReality_Left_Grip_Click;
+	static const FKey MixedReality_Left_Trigger_Click;
+	static const FKey MixedReality_Left_Trigger_Axis;
+	static const FKey MixedReality_Left_Thumbstick_X;
+	static const FKey MixedReality_Left_Thumbstick_Y;
+	static const FKey MixedReality_Left_Thumbstick_Click;
+	static const FKey MixedReality_Left_Trackpad_X;
+	static const FKey MixedReality_Left_Trackpad_Y;
+	static const FKey MixedReality_Left_Trackpad_Click;
+	static const FKey MixedReality_Left_Trackpad_Touch;
+	static const FKey MixedReality_Right_Menu_Click;
+	static const FKey MixedReality_Right_Grip_Click;
+	static const FKey MixedReality_Right_Trigger_Click;
+	static const FKey MixedReality_Right_Trigger_Axis;
+	static const FKey MixedReality_Right_Thumbstick_X;
+	static const FKey MixedReality_Right_Thumbstick_Y;
+	static const FKey MixedReality_Right_Thumbstick_Click;
+	static const FKey MixedReality_Right_Trackpad_X;
+	static const FKey MixedReality_Right_Trackpad_Y;
+	static const FKey MixedReality_Right_Trackpad_Click;
+	static const FKey MixedReality_Right_Trackpad_Touch;
+
+	// Oculus Go Controller
+	static const FKey OculusGo_Left_System_Click;
+	static const FKey OculusGo_Left_Back_Click;
+	static const FKey OculusGo_Left_Trigger_Click;
+	static const FKey OculusGo_Left_Trackpad_X;
+	static const FKey OculusGo_Left_Trackpad_Y;
+	static const FKey OculusGo_Left_Trackpad_Click;
+	static const FKey OculusGo_Left_Trackpad_Touch;
+	static const FKey OculusGo_Right_System_Click;
+	static const FKey OculusGo_Right_Back_Click;
+	static const FKey OculusGo_Right_Trigger_Click;
+	static const FKey OculusGo_Right_Trackpad_X;
+	static const FKey OculusGo_Right_Trackpad_Y;
+	static const FKey OculusGo_Right_Trackpad_Click;
+	static const FKey OculusGo_Right_Trackpad_Touch;
+
+	// Oculus Touch Controller
+	static const FKey OculusTouch_Left_X_Click;
+	static const FKey OculusTouch_Left_Y_Click;
+	static const FKey OculusTouch_Left_X_Touch;
+	static const FKey OculusTouch_Left_Y_Touch;
+	static const FKey OculusTouch_Left_Menu_Click;
+	static const FKey OculusTouch_Left_Grip_Click;
+	static const FKey OculusTouch_Left_Grip_Axis;
+	static const FKey OculusTouch_Left_Trigger_Click;
+	static const FKey OculusTouch_Left_Trigger_Axis;
+	static const FKey OculusTouch_Left_Trigger_Touch;
+	static const FKey OculusTouch_Left_Thumbstick_X;
+	static const FKey OculusTouch_Left_Thumbstick_Y;
+	static const FKey OculusTouch_Left_Thumbstick_Click;
+	static const FKey OculusTouch_Left_Thumbstick_Touch;
+	static const FKey OculusTouch_Right_A_Click;
+	static const FKey OculusTouch_Right_B_Click;
+	static const FKey OculusTouch_Right_A_Touch;
+	static const FKey OculusTouch_Right_B_Touch;
+	static const FKey OculusTouch_Right_System_Click;
+	static const FKey OculusTouch_Right_Grip_Click;
+	static const FKey OculusTouch_Right_Grip_Axis;
+	static const FKey OculusTouch_Right_Trigger_Click;
+	static const FKey OculusTouch_Right_Trigger_Axis;
+	static const FKey OculusTouch_Right_Trigger_Touch;
+	static const FKey OculusTouch_Right_Thumbstick_X;
+	static const FKey OculusTouch_Right_Thumbstick_Y;
+	static const FKey OculusTouch_Right_Thumbstick_Click;
+	static const FKey OculusTouch_Right_Thumbstick_Touch;
+
+	// Valve Index Controller
+	static const FKey ValveIndex_Left_A_Click;
+	static const FKey ValveIndex_Left_B_Click;
+	static const FKey ValveIndex_Left_A_Touch;
+	static const FKey ValveIndex_Left_B_Touch;
+	static const FKey ValveIndex_Left_System_Click;
+	static const FKey ValveIndex_Left_System_Touch;
+	static const FKey ValveIndex_Left_Grip_Click;
+	static const FKey ValveIndex_Left_Grip_Axis;
+	static const FKey ValveIndex_Left_Grip_Force;
+	static const FKey ValveIndex_Left_Trigger_Click;
+	static const FKey ValveIndex_Left_Trigger_Axis;
+	static const FKey ValveIndex_Left_Trigger_Touch;
+	static const FKey ValveIndex_Left_Thumbstick_X;
+	static const FKey ValveIndex_Left_Thumbstick_Y;
+	static const FKey ValveIndex_Left_Thumbstick_Click;
+	static const FKey ValveIndex_Left_Thumbstick_Touch;
+	static const FKey ValveIndex_Left_Trackpad_X;
+	static const FKey ValveIndex_Left_Trackpad_Y;
+	static const FKey ValveIndex_Left_Trackpad_Click;
+	static const FKey ValveIndex_Left_Trackpad_Force;
+	static const FKey ValveIndex_Left_Trackpad_Touch;
+	static const FKey ValveIndex_Right_A_Click;
+	static const FKey ValveIndex_Right_B_Click;
+	static const FKey ValveIndex_Right_A_Touch;
+	static const FKey ValveIndex_Right_B_Touch;
+	static const FKey ValveIndex_Right_System_Click;
+	static const FKey ValveIndex_Right_System_Touch;
+	static const FKey ValveIndex_Right_Grip_Click;
+	static const FKey ValveIndex_Right_Grip_Axis;
+	static const FKey ValveIndex_Right_Grip_Force;
+	static const FKey ValveIndex_Right_Trigger_Click;
+	static const FKey ValveIndex_Right_Trigger_Axis;
+	static const FKey ValveIndex_Right_Trigger_Touch;
+	static const FKey ValveIndex_Right_Thumbstick_X;
+	static const FKey ValveIndex_Right_Thumbstick_Y;
+	static const FKey ValveIndex_Right_Thumbstick_Click;
+	static const FKey ValveIndex_Right_Thumbstick_Touch;
+	static const FKey ValveIndex_Right_Trackpad_X;
+	static const FKey ValveIndex_Right_Trackpad_Y;
+	static const FKey ValveIndex_Right_Trackpad_Click;
+	static const FKey ValveIndex_Right_Trackpad_Force;
+	static const FKey ValveIndex_Right_Trackpad_Touch;
+
 	// Virtual buttons that use other buttons depending on the platform
 	static const FKey Virtual_Accept;
 	static const FKey Virtual_Back;
