Fixed missing header for non-unity build introduced by CL 9961639

diff --git a/Engine/Plugins/Runtime/Steam/SteamVR/Source/SteamVRInputDevice/Private/SteamVRInputDeviceModule.cpp b/Engine/Plugins/Runtime/Steam/SteamVR/Source/SteamVRInputDevice/Private/SteamVRInputDeviceModule.cpp
--- a/Engine/Plugins/Runtime/Steam/SteamVR/Source/SteamVRInputDevice/Private/SteamVRInputDeviceModule.cpp
+++ b/Engine/Plugins/Runtime/Steam/SteamVR/Source/SteamVRInputDevice/Private/SteamVRInputDeviceModule.cpp
@@ -33,6 +33,7 @@ POSSIBILITY OF SUCH DAMAGE.
 #include "Runtime/Projects/Public/Interfaces/IPluginManager.h"
 #include "SteamVRInputDevice.h"
 #include "Engine/Engine.h"
+#include "IXRTrackingSystem.h"
 
 class FSteamVRInputDeviceModule : public ISteamVRInputDeviceModule
 {
