Fix SteamVR Input static analysis warnings and copyright notices

diff --git a/Engine/Plugins/Runtime/Steam/SteamVR/Source/SteamVRInput/Private/SteamVRInput.cpp b/Engine/Plugins/Runtime/Steam/SteamVR/Source/SteamVRInput/Private/SteamVRInput.cpp
--- a/Engine/Plugins/Runtime/Steam/SteamVR/Source/SteamVRInput/Private/SteamVRInput.cpp
+++ b/Engine/Plugins/Runtime/Steam/SteamVR/Source/SteamVRInput/Private/SteamVRInput.cpp
@@ -1,32 +1,4 @@
-/*
-Copyright 2019 Valve Corporation under https://opensource.org/licenses/BSD-3-Clause
-
-Redistribution and use in source and binary forms, with or without modification,
-are permitted provided that the following conditions are met:
-
-1. Redistributions of source code must retain the above copyright notice, this
-   list of conditions and the following disclaimer.
-
-2. Redistributions in binary form must reproduce the above copyright notice,
-   this list of conditions and the following disclaimer in the documentation
-   and/or other materials provided with the distribution.
-
-3. Neither the name of the copyright holder nor the names of its contributors
-   may be used to endorse or promote products derived from this software
-   without specific prior written permission.
-
-THIS SOFTWARE IS PROVIDED BY THE COPYRIGHT HOLDERS AND CONTRIBUTORS "AS IS" AND
-ANY EXPRESS OR IMPLIED WARRANTIES, INCLUDING, BUT NOT LIMITED TO, THE IMPLIED
-WARRANTIES OF MERCHANTABILITY AND FITNESS FOR A PARTICULAR PURPOSE ARE DISCLAIMED.
-IN NO EVENT SHALL THE COPYRIGHT HOLDER OR CONTRIBUTORS BE LIABLE FOR ANY DIRECT,
-INDIRECT, INCIDENTAL, SPECIAL, EXEMPLARY, OR CONSEQUENTIAL DAMAGES (INCLUDING,
-BUT NOT LIMITED TO, PROCUREMENT OF SUBSTITUTE GOODS OR SERVICES; LOSS OF USE, DATA,
-OR PROFITS; OR BUSINESS INTERRUPTION) HOWEVER CAUSED AND ON ANY THEORY OF LIABILITY,
-WHETHER IN CONTRACT, STRICT LIABILITY, OR TORT (INCLUDING NEGLIGENCE OR OTHERWISE)
-ARISING IN ANY WAY OUT OF THE USE OF THIS SOFTWARE, EVEN IF ADVISED OF THE
-POSSIBILITY OF SUCH DAMAGE.
-*/
-
+// Copyright 1998-2019 Epic Games, Inc. All Rights Reserved.
 
 #include "SteamVRInput.h"
 #include "CoreMinimal.h"
diff --git a/Engine/Plugins/Runtime/Steam/SteamVR/Source/SteamVRInput/Public/SteamVRInput.h b/Engine/Plugins/Runtime/Steam/SteamVR/Source/SteamVRInput/Public/SteamVRInput.h
--- a/Engine/Plugins/Runtime/Steam/SteamVR/Source/SteamVRInput/Public/SteamVRInput.h
+++ b/Engine/Plugins/Runtime/Steam/SteamVR/Source/SteamVRInput/Public/SteamVRInput.h
@@ -1,32 +1,4 @@
-/*
-Copyright 2019 Valve Corporation under https://opensource.org/licenses/BSD-3-Clause
-
-Redistribution and use in source and binary forms, with or without modification,
-are permitted provided that the following conditions are met:
-
-1. Redistributions of source code must retain the above copyright notice, this
-   list of conditions and the following disclaimer.
-
-2. Redistributions in binary form must reproduce the above copyright notice,
-   this list of conditions and the following disclaimer in the documentation
-   and/or other materials provided with the distribution.
-
-3. Neither the name of the copyright holder nor the names of its contributors
-   may be used to endorse or promote products derived from this software
-   without specific prior written permission.
-
-THIS SOFTWARE IS PROVIDED BY THE COPYRIGHT HOLDERS AND CONTRIBUTORS "AS IS" AND
-ANY EXPRESS OR IMPLIED WARRANTIES, INCLUDING, BUT NOT LIMITED TO, THE IMPLIED
-WARRANTIES OF MERCHANTABILITY AND FITNESS FOR A PARTICULAR PURPOSE ARE DISCLAIMED.
-IN NO EVENT SHALL THE COPYRIGHT HOLDER OR CONTRIBUTORS BE LIABLE FOR ANY DIRECT,
-INDIRECT, INCIDENTAL, SPECIAL, EXEMPLARY, OR CONSEQUENTIAL DAMAGES (INCLUDING,
-BUT NOT LIMITED TO, PROCUREMENT OF SUBSTITUTE GOODS OR SERVICES; LOSS OF USE, DATA,
-OR PROFITS; OR BUSINESS INTERRUPTION) HOWEVER CAUSED AND ON ANY THEORY OF LIABILITY,
-WHETHER IN CONTRACT, STRICT LIABILITY, OR TORT (INCLUDING NEGLIGENCE OR OTHERWISE)
-ARISING IN ANY WAY OUT OF THE USE OF THIS SOFTWARE, EVEN IF ADVISED OF THE
-POSSIBILITY OF SUCH DAMAGE.
-*/
-
+// Copyright 1998-2019 Epic Games, Inc. All Rights Reserved.
 
 #pragma once
 
diff --git a/Engine/Plugins/Runtime/Steam/SteamVR/Source/SteamVRInputDevice/Private/AnimNode_SteamVRInputAnimPose.cpp b/Engine/Plugins/Runtime/Steam/SteamVR/Source/SteamVRInputDevice/Private/AnimNode_SteamVRInputAnimPose.cpp
--- a/Engine/Plugins/Runtime/Steam/SteamVR/Source/SteamVRInputDevice/Private/AnimNode_SteamVRInputAnimPose.cpp
+++ b/Engine/Plugins/Runtime/Steam/SteamVR/Source/SteamVRInputDevice/Private/AnimNode_SteamVRInputAnimPose.cpp
@@ -1,32 +1,4 @@
-/*
-Copyright 2019 Valve Corporation under https://opensource.org/licenses/BSD-3-Clause
-
-Redistribution and use in source and binary forms, with or without modification,
-are permitted provided that the following conditions are met:
-
-1. Redistributions of source code must retain the above copyright notice, this
-   list of conditions and the following disclaimer.
-
-2. Redistributions in binary form must reproduce the above copyright notice,
-   this list of conditions and the following disclaimer in the documentation
-   and/or other materials provided with the distribution.
-
-3. Neither the name of the copyright holder nor the names of its contributors
-   may be used to endorse or promote products derived from this software
-   without specific prior written permission.
-
-THIS SOFTWARE IS PROVIDED BY THE COPYRIGHT HOLDERS AND CONTRIBUTORS "AS IS" AND
-ANY EXPRESS OR IMPLIED WARRANTIES, INCLUDING, BUT NOT LIMITED TO, THE IMPLIED
-WARRANTIES OF MERCHANTABILITY AND FITNESS FOR A PARTICULAR PURPOSE ARE DISCLAIMED.
-IN NO EVENT SHALL THE COPYRIGHT HOLDER OR CONTRIBUTORS BE LIABLE FOR ANY DIRECT,
-INDIRECT, INCIDENTAL, SPECIAL, EXEMPLARY, OR CONSEQUENTIAL DAMAGES (INCLUDING,
-BUT NOT LIMITED TO, PROCUREMENT OF SUBSTITUTE GOODS OR SERVICES; LOSS OF USE, DATA,
-OR PROFITS; OR BUSINESS INTERRUPTION) HOWEVER CAUSED AND ON ANY THEORY OF LIABILITY,
-WHETHER IN CONTRACT, STRICT LIABILITY, OR TORT (INCLUDING NEGLIGENCE OR OTHERWISE)
-ARISING IN ANY WAY OUT OF THE USE OF THIS SOFTWARE, EVEN IF ADVISED OF THE
-POSSIBILITY OF SUCH DAMAGE.
-*/
-
+// Copyright 1998-2019 Epic Games, Inc. All Rights Reserved.
 
 #include "AnimNode_SteamVRInputAnimPose.h"
 #include "ISteamVRInputDeviceModule.h"
diff --git a/Engine/Plugins/Runtime/Steam/SteamVR/Source/SteamVRInputDevice/Private/AnimNode_SteamVRSetWristTransform.cpp b/Engine/Plugins/Runtime/Steam/SteamVR/Source/SteamVRInputDevice/Private/AnimNode_SteamVRSetWristTransform.cpp
--- a/Engine/Plugins/Runtime/Steam/SteamVR/Source/SteamVRInputDevice/Private/AnimNode_SteamVRSetWristTransform.cpp
+++ b/Engine/Plugins/Runtime/Steam/SteamVR/Source/SteamVRInputDevice/Private/AnimNode_SteamVRSetWristTransform.cpp
@@ -1,32 +1,4 @@
-/*
-Copyright 2019 Valve Corporation under https://opensource.org/licenses/BSD-3-Clause
-
-Redistribution and use in source and binary forms, with or without modification,
-are permitted provided that the following conditions are met:
-
-1. Redistributions of source code must retain the above copyright notice, this
-   list of conditions and the following disclaimer.
-
-2. Redistributions in binary form must reproduce the above copyright notice,
-   this list of conditions and the following disclaimer in the documentation
-   and/or other materials provided with the distribution.
-
-3. Neither the name of the copyright holder nor the names of its contributors
-   may be used to endorse or promote products derived from this software
-   without specific prior written permission.
-
-THIS SOFTWARE IS PROVIDED BY THE COPYRIGHT HOLDERS AND CONTRIBUTORS "AS IS" AND
-ANY EXPRESS OR IMPLIED WARRANTIES, INCLUDING, BUT NOT LIMITED TO, THE IMPLIED
-WARRANTIES OF MERCHANTABILITY AND FITNESS FOR A PARTICULAR PURPOSE ARE DISCLAIMED.
-IN NO EVENT SHALL THE COPYRIGHT HOLDER OR CONTRIBUTORS BE LIABLE FOR ANY DIRECT,
-INDIRECT, INCIDENTAL, SPECIAL, EXEMPLARY, OR CONSEQUENTIAL DAMAGES (INCLUDING,
-BUT NOT LIMITED TO, PROCUREMENT OF SUBSTITUTE GOODS OR SERVICES; LOSS OF USE, DATA,
-OR PROFITS; OR BUSINESS INTERRUPTION) HOWEVER CAUSED AND ON ANY THEORY OF LIABILITY,
-WHETHER IN CONTRACT, STRICT LIABILITY, OR TORT (INCLUDING NEGLIGENCE OR OTHERWISE)
-ARISING IN ANY WAY OUT OF THE USE OF THIS SOFTWARE, EVEN IF ADVISED OF THE
-POSSIBILITY OF SUCH DAMAGE.
-*/
-
+// Copyright 1998-2019 Epic Games, Inc. All Rights Reserved.
 
 #include "AnimNode_SteamVRSetWristTransform.h"
 #include "ISteamVRInputDeviceModule.h"
diff --git a/Engine/Plugins/Runtime/Steam/SteamVR/Source/SteamVRInputDevice/Private/SteamVRInputDevice.cpp b/Engine/Plugins/Runtime/Steam/SteamVR/Source/SteamVRInputDevice/Private/SteamVRInputDevice.cpp
--- a/Engine/Plugins/Runtime/Steam/SteamVR/Source/SteamVRInputDevice/Private/SteamVRInputDevice.cpp
+++ b/Engine/Plugins/Runtime/Steam/SteamVR/Source/SteamVRInputDevice/Private/SteamVRInputDevice.cpp
@@ -1,32 +1,4 @@
-/*
-Copyright 2019 Valve Corporation under https://opensource.org/licenses/BSD-3-Clause
-
-Redistribution and use in source and binary forms, with or without modification,
-are permitted provided that the following conditions are met:
-
-1. Redistributions of source code must retain the above copyright notice, this
-   list of conditions and the following disclaimer.
-
-2. Redistributions in binary form must reproduce the above copyright notice,
-   this list of conditions and the following disclaimer in the documentation
-   and/or other materials provided with the distribution.
-
-3. Neither the name of the copyright holder nor the names of its contributors
-   may be used to endorse or promote products derived from this software
-   without specific prior written permission.
-
-THIS SOFTWARE IS PROVIDED BY THE COPYRIGHT HOLDERS AND CONTRIBUTORS "AS IS" AND
-ANY EXPRESS OR IMPLIED WARRANTIES, INCLUDING, BUT NOT LIMITED TO, THE IMPLIED
-WARRANTIES OF MERCHANTABILITY AND FITNESS FOR A PARTICULAR PURPOSE ARE DISCLAIMED.
-IN NO EVENT SHALL THE COPYRIGHT HOLDER OR CONTRIBUTORS BE LIABLE FOR ANY DIRECT,
-INDIRECT, INCIDENTAL, SPECIAL, EXEMPLARY, OR CONSEQUENTIAL DAMAGES (INCLUDING,
-BUT NOT LIMITED TO, PROCUREMENT OF SUBSTITUTE GOODS OR SERVICES; LOSS OF USE, DATA,
-OR PROFITS; OR BUSINESS INTERRUPTION) HOWEVER CAUSED AND ON ANY THEORY OF LIABILITY,
-WHETHER IN CONTRACT, STRICT LIABILITY, OR TORT (INCLUDING NEGLIGENCE OR OTHERWISE)
-ARISING IN ANY WAY OUT OF THE USE OF THIS SOFTWARE, EVEN IF ADVISED OF THE
-POSSIBILITY OF SUCH DAMAGE.
-*/
-
+// Copyright 1998-2019 Epic Games, Inc. All Rights Reserved.
 
 #include "SteamVRInputDevice.h"
 #include "Runtime/ApplicationCore/Public/GenericPlatform/IInputInterface.h"
@@ -978,7 +950,7 @@ void FSteamVRInputDevice::GetControllerFidelity()
 				return;
 			}
 
-			VRInput()->GetSkeletalTrackingLevel(VRSkeletalHandleLeft, &LeftControllerFidelity);
+			InputError = VRInput()->GetSkeletalTrackingLevel(VRSkeletalHandleLeft, &LeftControllerFidelity);
 
 			if (InputError != VRInputError_None)
 			{
diff --git a/Engine/Plugins/Runtime/Steam/SteamVR/Source/SteamVRInputDevice/Private/SteamVRInputDeviceFunctionLibrary.cpp b/Engine/Plugins/Runtime/Steam/SteamVR/Source/SteamVRInputDevice/Private/SteamVRInputDeviceFunctionLibrary.cpp
--- a/Engine/Plugins/Runtime/Steam/SteamVR/Source/SteamVRInputDevice/Private/SteamVRInputDeviceFunctionLibrary.cpp
+++ b/Engine/Plugins/Runtime/Steam/SteamVR/Source/SteamVRInputDevice/Private/SteamVRInputDeviceFunctionLibrary.cpp
@@ -1,32 +1,4 @@
-/*
-Copyright 2019 Valve Corporation under https://opensource.org/licenses/BSD-3-Clause
-
-Redistribution and use in source and binary forms, with or without modification,
-are permitted provided that the following conditions are met:
-
-1. Redistributions of source code must retain the above copyright notice, this
-   list of conditions and the following disclaimer.
-
-2. Redistributions in binary form must reproduce the above copyright notice,
-   this list of conditions and the following disclaimer in the documentation
-   and/or other materials provided with the distribution.
-
-3. Neither the name of the copyright holder nor the names of its contributors
-   may be used to endorse or promote products derived from this software
-   without specific prior written permission.
-
-THIS SOFTWARE IS PROVIDED BY THE COPYRIGHT HOLDERS AND CONTRIBUTORS "AS IS" AND
-ANY EXPRESS OR IMPLIED WARRANTIES, INCLUDING, BUT NOT LIMITED TO, THE IMPLIED
-WARRANTIES OF MERCHANTABILITY AND FITNESS FOR A PARTICULAR PURPOSE ARE DISCLAIMED.
-IN NO EVENT SHALL THE COPYRIGHT HOLDER OR CONTRIBUTORS BE LIABLE FOR ANY DIRECT,
-INDIRECT, INCIDENTAL, SPECIAL, EXEMPLARY, OR CONSEQUENTIAL DAMAGES (INCLUDING,
-BUT NOT LIMITED TO, PROCUREMENT OF SUBSTITUTE GOODS OR SERVICES; LOSS OF USE, DATA,
-OR PROFITS; OR BUSINESS INTERRUPTION) HOWEVER CAUSED AND ON ANY THEORY OF LIABILITY,
-WHETHER IN CONTRACT, STRICT LIABILITY, OR TORT (INCLUDING NEGLIGENCE OR OTHERWISE)
-ARISING IN ANY WAY OUT OF THE USE OF THIS SOFTWARE, EVEN IF ADVISED OF THE
-POSSIBILITY OF SUCH DAMAGE.
-*/
-
+// Copyright 1998-2019 Epic Games, Inc. All Rights Reserved.
 
 #include "SteamVRInputDeviceFunctionLibrary.h"
 
@@ -56,21 +28,11 @@ void USteamVRInputDeviceFunctionLibrary::PlaySteamVR_HapticFeedback(ESteamVRHand
 
 		if (Hand == ESteamVRHand::VR_Left && SteamVRInputDevice->VRVibrationLeft != k_ulInvalidActionHandle && SteamVRInputDevice->bIsSkeletalControllerLeftPresent && SteamVRInputDevice->VRSkeletalHandleLeft != k_ulInvalidActionHandle)
 		{
-			if (SteamVRInputDevice->VRVibrationLeft == k_ulInvalidActionHandle)
-			{
-				return;
-			}
-
 			ActiveSkeletalHand = SteamVRInputDevice->VRSkeletalHandleLeft;
 			VRInput()->TriggerHapticVibrationAction(SteamVRInputDevice->VRVibrationLeft, StartSecondsFromNow, DurationSeconds, Frequency, Amplitude, k_ulInvalidInputValueHandle);
 		}
-		else if (Hand == ESteamVRHand::VR_Right && SteamVRInputDevice->VRVibrationLeft != k_ulInvalidActionHandle && SteamVRInputDevice->bIsSkeletalControllerRightPresent && SteamVRInputDevice->VRSkeletalHandleRight != k_ulInvalidActionHandle)
+		else if (Hand == ESteamVRHand::VR_Right && SteamVRInputDevice->VRVibrationRight != k_ulInvalidActionHandle && SteamVRInputDevice->bIsSkeletalControllerRightPresent && SteamVRInputDevice->VRSkeletalHandleRight != k_ulInvalidActionHandle)
 		{
-			if (SteamVRInputDevice->VRVibrationRight == k_ulInvalidActionHandle)
-			{
-				return;
-			}
-
 			ActiveSkeletalHand = SteamVRInputDevice->VRSkeletalHandleRight;
 			VRInput()->TriggerHapticVibrationAction(SteamVRInputDevice->VRVibrationRight, StartSecondsFromNow, DurationSeconds, Frequency, Amplitude, k_ulInvalidInputValueHandle);
 		}
@@ -255,8 +217,6 @@ void USteamVRInputDeviceFunctionLibrary::GetSkeletalTransform(FSteamVRSkeletonTr
 		LeftHand.Aux_Ring = OutPose[29];
 		LeftHand.Aux_Pinky = OutPose[30];
 
-		LeftHand.Bone_Count = OutPose[31];
-
 		// Right Hand - Grab Skeletal Data
 		SteamVRInputDevice->GetSkeletalData(false, false, MotionRange, OutPose, STEAMVR_SKELETON_BONE_COUNT);
 
@@ -297,8 +257,6 @@ void USteamVRInputDeviceFunctionLibrary::GetSkeletalTransform(FSteamVRSkeletonTr
 		RightHand.Aux_Middle = OutPose[28];
 		RightHand.Aux_Ring = OutPose[29];
 		RightHand.Aux_Pinky = OutPose[30];
-
-		RightHand.Bone_Count = OutPose[31];
 	}
 }
 
@@ -498,9 +456,9 @@ void USteamVRInputDeviceFunctionLibrary::ShowSteamVR_ActionOrigin(FSteamVRAction
 
 		if (GetSteamVR_OriginTrackedDeviceInfo(SteamVRAction, OriginInfo))
 		{
-			VRActiveActionSet_t ActiveActionSets[] = { 0 };
-			ActiveActionSets[0].ulActionSet = SteamVRActionSet.Handle;
-			VRInput()->ShowBindingsForActionSet(ActiveActionSets, sizeof(ActiveActionSets[0]), 1, SteamVRAction.ActiveOrigin);
+			VRActiveActionSet_t ActiveActionSet = { 0 };
+			ActiveActionSet.ulActionSet = SteamVRActionSet.Handle;
+			VRInput()->ShowBindingsForActionSet(&ActiveActionSet, sizeof(ActiveActionSet), 1, SteamVRAction.ActiveOrigin);
 
 			UE_LOG(LogTemp, Warning, TEXT("Action [%s] triggered from Device [%i][%s] at Component [%s]"), *SteamVRAction.Name.ToString(), OriginInfo.TrackedDeviceIndex, *OriginInfo.TrackedDeviceModel, *OriginInfo.RenderModelComponentName);
 		}
@@ -608,8 +566,8 @@ float USteamVRInputDeviceFunctionLibrary::SetSteamVR_GlobalPredictedSecondsFromN
 
 void USteamVRInputDeviceFunctionLibrary::ShowAllSteamVR_ActionOrigins()
 {
-	VRActiveActionSet_t ActiveActionSets[1];
-	VRInput()->ShowBindingsForActionSet(ActiveActionSets, sizeof(ActiveActionSets[0]), 0, 0);
+	VRActiveActionSet_t ActiveActionSet;
+	VRInput()->ShowBindingsForActionSet(&ActiveActionSet, sizeof(ActiveActionSet), 1, 0);
 }
 
 TArray<FSteamVRInputBindingInfo> USteamVRInputDeviceFunctionLibrary::GetSteamVR_InputBindingInfo(FSteamVRAction SteamVRActionHandle)
@@ -768,15 +726,13 @@ void USteamVRInputDeviceFunctionLibrary::GetFingerCurlsAndSplays(EHand Hand, FSt
 	if (SteamVRInputDevice != nullptr && VRSystem() && VRInput())
 	{
 		// Get action state this frame
-		VRActiveActionSet_t ActiveActionSets[] = {
-			{
-				SteamVRInputDevice->MainActionSet,
-				k_ulInvalidInputValueHandle,
-				k_ulInvalidActionSetHandle
-			}
+		VRActiveActionSet_t ActiveActionSet = {
+			SteamVRInputDevice->MainActionSet,
+			k_ulInvalidInputValueHandle,
+			k_ulInvalidActionSetHandle
 		};
 
-		EVRInputError UpdateActionStateError = VRInput()->UpdateActionState(ActiveActionSets, sizeof(VRActiveActionSet_t), 1);
+		EVRInputError UpdateActionStateError = VRInput()->UpdateActionState(&ActiveActionSet, sizeof(VRActiveActionSet_t), 1);
 		if (UpdateActionStateError != VRInputError_None)
 		{
 			FingerCurls = {};
@@ -832,13 +788,7 @@ void USteamVRInputDeviceFunctionLibrary::GetFingerCurlsAndSplays(EHand Hand, FSt
 		EVRSummaryType SteamVRSummaryType = (SummaryDataType == ESkeletalSummaryDataType::VR_SummaryType_FromDevice) ? VRSummaryType_FromDevice : VRSummaryType_FromAnimation;
 		EVRInputError GetSkeletalSummaryDataError = VRInput()->GetSkeletalSummaryData(ActiveSkeletalHand, SteamVRSummaryType, &ActiveSkeletalSummaryData);
 
-		if (GetSkeletalSummaryDataError != VRInputError_None)
-		{
-			FingerCurls = {};
-			FingerSplays = {};
-			return;
-		}
-		else if (GetSkeletalSummaryDataError == VRInputError_None)
+		if (GetSkeletalSummaryDataError == VRInputError_None)
 		{
 			// Update curls and splay values for output
 			FingerCurls.Thumb = ActiveSkeletalSummaryData.flFingerCurl[VRFinger_Thumb];
@@ -853,6 +803,12 @@ void USteamVRInputDeviceFunctionLibrary::GetFingerCurlsAndSplays(EHand Hand, FSt
 			FingerSplays.Thumb_Index = ActiveSkeletalSummaryData.flFingerSplay[VRFingerSplay_Thumb_Index];
 			return;
 		}
+		else
+		{
+			FingerCurls = {};
+			FingerSplays = {};
+			return;
+		}
 	}
 
 	// Unable to retrieve the curls and splay values for this hand, send zeroed out values back to the user
diff --git a/Engine/Plugins/Runtime/Steam/SteamVR/Source/SteamVRInputDevice/Private/SteamVRInputDeviceModule.cpp b/Engine/Plugins/Runtime/Steam/SteamVR/Source/SteamVRInputDevice/Private/SteamVRInputDeviceModule.cpp
--- a/Engine/Plugins/Runtime/Steam/SteamVR/Source/SteamVRInputDevice/Private/SteamVRInputDeviceModule.cpp
+++ b/Engine/Plugins/Runtime/Steam/SteamVR/Source/SteamVRInputDevice/Private/SteamVRInputDeviceModule.cpp
@@ -1,32 +1,4 @@
-/*
-Copyright 2019 Valve Corporation under https://opensource.org/licenses/BSD-3-Clause
-
-Redistribution and use in source and binary forms, with or without modification,
-are permitted provided that the following conditions are met:
-
-1. Redistributions of source code must retain the above copyright notice, this
-   list of conditions and the following disclaimer.
-
-2. Redistributions in binary form must reproduce the above copyright notice,
-   this list of conditions and the following disclaimer in the documentation
-   and/or other materials provided with the distribution.
-
-3. Neither the name of the copyright holder nor the names of its contributors
-   may be used to endorse or promote products derived from this software
-   without specific prior written permission.
-
-THIS SOFTWARE IS PROVIDED BY THE COPYRIGHT HOLDERS AND CONTRIBUTORS "AS IS" AND
-ANY EXPRESS OR IMPLIED WARRANTIES, INCLUDING, BUT NOT LIMITED TO, THE IMPLIED
-WARRANTIES OF MERCHANTABILITY AND FITNESS FOR A PARTICULAR PURPOSE ARE DISCLAIMED.
-IN NO EVENT SHALL THE COPYRIGHT HOLDER OR CONTRIBUTORS BE LIABLE FOR ANY DIRECT,
-INDIRECT, INCIDENTAL, SPECIAL, EXEMPLARY, OR CONSEQUENTIAL DAMAGES (INCLUDING,
-BUT NOT LIMITED TO, PROCUREMENT OF SUBSTITUTE GOODS OR SERVICES; LOSS OF USE, DATA,
-OR PROFITS; OR BUSINESS INTERRUPTION) HOWEVER CAUSED AND ON ANY THEORY OF LIABILITY,
-WHETHER IN CONTRACT, STRICT LIABILITY, OR TORT (INCLUDING NEGLIGENCE OR OTHERWISE)
-ARISING IN ANY WAY OUT OF THE USE OF THIS SOFTWARE, EVEN IF ADVISED OF THE
-POSSIBILITY OF SUCH DAMAGE.
-*/
-
+// Copyright 1998-2019 Epic Games, Inc. All Rights Reserved.
 
 #include "ISteamVRInputDeviceModule.h"
 #include "IInputDevice.h"
diff --git a/Engine/Plugins/Runtime/Steam/SteamVR/Source/SteamVRInputDevice/Private/SteamVRSkeletonDefinition.cpp b/Engine/Plugins/Runtime/Steam/SteamVR/Source/SteamVRInputDevice/Private/SteamVRSkeletonDefinition.cpp
--- a/Engine/Plugins/Runtime/Steam/SteamVR/Source/SteamVRInputDevice/Private/SteamVRSkeletonDefinition.cpp
+++ b/Engine/Plugins/Runtime/Steam/SteamVR/Source/SteamVRInputDevice/Private/SteamVRSkeletonDefinition.cpp
@@ -1,32 +1,4 @@
-/*
-Copyright 2019 Valve Corporation under https://opensource.org/licenses/BSD-3-Clause
-
-Redistribution and use in source and binary forms, with or without modification,
-are permitted provided that the following conditions are met:
-
-1. Redistributions of source code must retain the above copyright notice, this
-   list of conditions and the following disclaimer.
-
-2. Redistributions in binary form must reproduce the above copyright notice,
-   this list of conditions and the following disclaimer in the documentation
-   and/or other materials provided with the distribution.
-
-3. Neither the name of the copyright holder nor the names of its contributors
-   may be used to endorse or promote products derived from this software
-   without specific prior written permission.
-
-THIS SOFTWARE IS PROVIDED BY THE COPYRIGHT HOLDERS AND CONTRIBUTORS "AS IS" AND
-ANY EXPRESS OR IMPLIED WARRANTIES, INCLUDING, BUT NOT LIMITED TO, THE IMPLIED
-WARRANTIES OF MERCHANTABILITY AND FITNESS FOR A PARTICULAR PURPOSE ARE DISCLAIMED.
-IN NO EVENT SHALL THE COPYRIGHT HOLDER OR CONTRIBUTORS BE LIABLE FOR ANY DIRECT,
-INDIRECT, INCIDENTAL, SPECIAL, EXEMPLARY, OR CONSEQUENTIAL DAMAGES (INCLUDING,
-BUT NOT LIMITED TO, PROCUREMENT OF SUBSTITUTE GOODS OR SERVICES; LOSS OF USE, DATA,
-OR PROFITS; OR BUSINESS INTERRUPTION) HOWEVER CAUSED AND ON ANY THEORY OF LIABILITY,
-WHETHER IN CONTRACT, STRICT LIABILITY, OR TORT (INCLUDING NEGLIGENCE OR OTHERWISE)
-ARISING IN ANY WAY OUT OF THE USE OF THIS SOFTWARE, EVEN IF ADVISED OF THE
-POSSIBILITY OF SUCH DAMAGE.
-*/
-
+// Copyright 1998-2019 Epic Games, Inc. All Rights Reserved.
 
 #include "SteamVRSkeletonDefinition.h"
 
diff --git a/Engine/Plugins/Runtime/Steam/SteamVR/Source/SteamVRInputDevice/Private/SteamVRTrackingRefComponent.cpp b/Engine/Plugins/Runtime/Steam/SteamVR/Source/SteamVRInputDevice/Private/SteamVRTrackingRefComponent.cpp
--- a/Engine/Plugins/Runtime/Steam/SteamVR/Source/SteamVRInputDevice/Private/SteamVRTrackingRefComponent.cpp
+++ b/Engine/Plugins/Runtime/Steam/SteamVR/Source/SteamVRInputDevice/Private/SteamVRTrackingRefComponent.cpp
@@ -1,31 +1,4 @@
-/*
-Copyright 2019 Valve Corporation under https://opensource.org/licenses/BSD-3-Clause
-
-Redistribution and use in source and binary forms, with or without modification,
-are permitted provided that the following conditions are met:
-
-1. Redistributions of source code must retain the above copyright notice, this
-   list of conditions and the following disclaimer.
-
-2. Redistributions in binary form must reproduce the above copyright notice,
-   this list of conditions and the following disclaimer in the documentation
-   and/or other materials provided with the distribution.
-
-3. Neither the name of the copyright holder nor the names of its contributors
-   may be used to endorse or promote products derived from this software
-   without specific prior written permission.
-
-THIS SOFTWARE IS PROVIDED BY THE COPYRIGHT HOLDERS AND CONTRIBUTORS "AS IS" AND
-ANY EXPRESS OR IMPLIED WARRANTIES, INCLUDING, BUT NOT LIMITED TO, THE IMPLIED
-WARRANTIES OF MERCHANTABILITY AND FITNESS FOR A PARTICULAR PURPOSE ARE DISCLAIMED.
-IN NO EVENT SHALL THE COPYRIGHT HOLDER OR CONTRIBUTORS BE LIABLE FOR ANY DIRECT,
-INDIRECT, INCIDENTAL, SPECIAL, EXEMPLARY, OR CONSEQUENTIAL DAMAGES (INCLUDING,
-BUT NOT LIMITED TO, PROCUREMENT OF SUBSTITUTE GOODS OR SERVICES; LOSS OF USE, DATA,
-OR PROFITS; OR BUSINESS INTERRUPTION) HOWEVER CAUSED AND ON ANY THEORY OF LIABILITY,
-WHETHER IN CONTRACT, STRICT LIABILITY, OR TORT (INCLUDING NEGLIGENCE OR OTHERWISE)
-ARISING IN ANY WAY OUT OF THE USE OF THIS SOFTWARE, EVEN IF ADVISED OF THE
-POSSIBILITY OF SUCH DAMAGE.
-*/
+// Copyright 1998-2019 Epic Games, Inc. All Rights Reserved.
 
 #include "SteamVRTrackingRefComponent.h"
 #include "openvr.h"
@@ -70,7 +43,7 @@ bool USteamVRTrackingReferences::ShowTrackingReferences(UStaticMesh* TrackingRef
 				//FString stringCache = *FString(UTF8_TO_TCHAR(buf));
 				//UE_LOG(LogSteamVRTrackingRefComponent, Warning, TEXT("[TRACKING REFERENCE] Found the following tracking device: [%i] %s"), id, *stringCache);
 
-				TrackedDevicePose_t TrackedDevicePose = { 0 };
+				TrackedDevicePose_t TrackedDevicePose;
 				VRCompositor()->GetLastPoseForTrackedDeviceIndex(id, &TrackedDevicePose, NULL);
 
 				// Get SteamVR Transform Matrix for this tracking reference
diff --git a/Engine/Plugins/Runtime/Steam/SteamVR/Source/SteamVRInputDevice/Private/UE4HandSkeletonDefinition.cpp b/Engine/Plugins/Runtime/Steam/SteamVR/Source/SteamVRInputDevice/Private/UE4HandSkeletonDefinition.cpp
--- a/Engine/Plugins/Runtime/Steam/SteamVR/Source/SteamVRInputDevice/Private/UE4HandSkeletonDefinition.cpp
+++ b/Engine/Plugins/Runtime/Steam/SteamVR/Source/SteamVRInputDevice/Private/UE4HandSkeletonDefinition.cpp
@@ -1,32 +1,4 @@
-/*
-Copyright 2019 Valve Corporation under https://opensource.org/licenses/BSD-3-Clause
-
-Redistribution and use in source and binary forms, with or without modification,
-are permitted provided that the following conditions are met:
-
-1. Redistributions of source code must retain the above copyright notice, this
-   list of conditions and the following disclaimer.
-
-2. Redistributions in binary form must reproduce the above copyright notice,
-   this list of conditions and the following disclaimer in the documentation
-   and/or other materials provided with the distribution.
-
-3. Neither the name of the copyright holder nor the names of its contributors
-   may be used to endorse or promote products derived from this software
-   without specific prior written permission.
-
-THIS SOFTWARE IS PROVIDED BY THE COPYRIGHT HOLDERS AND CONTRIBUTORS "AS IS" AND
-ANY EXPRESS OR IMPLIED WARRANTIES, INCLUDING, BUT NOT LIMITED TO, THE IMPLIED
-WARRANTIES OF MERCHANTABILITY AND FITNESS FOR A PARTICULAR PURPOSE ARE DISCLAIMED.
-IN NO EVENT SHALL THE COPYRIGHT HOLDER OR CONTRIBUTORS BE LIABLE FOR ANY DIRECT,
-INDIRECT, INCIDENTAL, SPECIAL, EXEMPLARY, OR CONSEQUENTIAL DAMAGES (INCLUDING,
-BUT NOT LIMITED TO, PROCUREMENT OF SUBSTITUTE GOODS OR SERVICES; LOSS OF USE, DATA,
-OR PROFITS; OR BUSINESS INTERRUPTION) HOWEVER CAUSED AND ON ANY THEORY OF LIABILITY,
-WHETHER IN CONTRACT, STRICT LIABILITY, OR TORT (INCLUDING NEGLIGENCE OR OTHERWISE)
-ARISING IN ANY WAY OUT OF THE USE OF THIS SOFTWARE, EVEN IF ADVISED OF THE
-POSSIBILITY OF SUCH DAMAGE.
-*/
-
+// Copyright 1998-2019 Epic Games, Inc. All Rights Reserved.
 
 #include "UE4HandSkeletonDefinition.h"
 
diff --git a/Engine/Plugins/Runtime/Steam/SteamVR/Source/SteamVRInputDevice/Public/AnimNode_SteamVRInputAnimPose.h b/Engine/Plugins/Runtime/Steam/SteamVR/Source/SteamVRInputDevice/Public/AnimNode_SteamVRInputAnimPose.h
--- a/Engine/Plugins/Runtime/Steam/SteamVR/Source/SteamVRInputDevice/Public/AnimNode_SteamVRInputAnimPose.h
+++ b/Engine/Plugins/Runtime/Steam/SteamVR/Source/SteamVRInputDevice/Public/AnimNode_SteamVRInputAnimPose.h
@@ -1,32 +1,4 @@
-/*
-Copyright 2019 Valve Corporation under https://opensource.org/licenses/BSD-3-Clause
-
-Redistribution and use in source and binary forms, with or without modification,
-are permitted provided that the following conditions are met:
-
-1. Redistributions of source code must retain the above copyright notice, this
-   list of conditions and the following disclaimer.
-
-2. Redistributions in binary form must reproduce the above copyright notice,
-   this list of conditions and the following disclaimer in the documentation
-   and/or other materials provided with the distribution.
-
-3. Neither the name of the copyright holder nor the names of its contributors
-   may be used to endorse or promote products derived from this software
-   without specific prior written permission.
-
-THIS SOFTWARE IS PROVIDED BY THE COPYRIGHT HOLDERS AND CONTRIBUTORS "AS IS" AND
-ANY EXPRESS OR IMPLIED WARRANTIES, INCLUDING, BUT NOT LIMITED TO, THE IMPLIED
-WARRANTIES OF MERCHANTABILITY AND FITNESS FOR A PARTICULAR PURPOSE ARE DISCLAIMED.
-IN NO EVENT SHALL THE COPYRIGHT HOLDER OR CONTRIBUTORS BE LIABLE FOR ANY DIRECT,
-INDIRECT, INCIDENTAL, SPECIAL, EXEMPLARY, OR CONSEQUENTIAL DAMAGES (INCLUDING,
-BUT NOT LIMITED TO, PROCUREMENT OF SUBSTITUTE GOODS OR SERVICES; LOSS OF USE, DATA,
-OR PROFITS; OR BUSINESS INTERRUPTION) HOWEVER CAUSED AND ON ANY THEORY OF LIABILITY,
-WHETHER IN CONTRACT, STRICT LIABILITY, OR TORT (INCLUDING NEGLIGENCE OR OTHERWISE)
-ARISING IN ANY WAY OUT OF THE USE OF THIS SOFTWARE, EVEN IF ADVISED OF THE
-POSSIBILITY OF SUCH DAMAGE.
-*/
-
+// Copyright 1998-2019 Epic Games, Inc. All Rights Reserved.
 
 #pragma once
 
diff --git a/Engine/Plugins/Runtime/Steam/SteamVR/Source/SteamVRInputDevice/Public/AnimNode_SteamVRSetWristTransform.h b/Engine/Plugins/Runtime/Steam/SteamVR/Source/SteamVRInputDevice/Public/AnimNode_SteamVRSetWristTransform.h
--- a/Engine/Plugins/Runtime/Steam/SteamVR/Source/SteamVRInputDevice/Public/AnimNode_SteamVRSetWristTransform.h
+++ b/Engine/Plugins/Runtime/Steam/SteamVR/Source/SteamVRInputDevice/Public/AnimNode_SteamVRSetWristTransform.h
@@ -1,32 +1,4 @@
-/*
-Copyright 2019 Valve Corporation under https://opensource.org/licenses/BSD-3-Clause
-
-Redistribution and use in source and binary forms, with or without modification,
-are permitted provided that the following conditions are met:
-
-1. Redistributions of source code must retain the above copyright notice, this
-   list of conditions and the following disclaimer.
-
-2. Redistributions in binary form must reproduce the above copyright notice,
-   this list of conditions and the following disclaimer in the documentation
-   and/or other materials provided with the distribution.
-
-3. Neither the name of the copyright holder nor the names of its contributors
-   may be used to endorse or promote products derived from this software
-   without specific prior written permission.
-
-THIS SOFTWARE IS PROVIDED BY THE COPYRIGHT HOLDERS AND CONTRIBUTORS "AS IS" AND
-ANY EXPRESS OR IMPLIED WARRANTIES, INCLUDING, BUT NOT LIMITED TO, THE IMPLIED
-WARRANTIES OF MERCHANTABILITY AND FITNESS FOR A PARTICULAR PURPOSE ARE DISCLAIMED.
-IN NO EVENT SHALL THE COPYRIGHT HOLDER OR CONTRIBUTORS BE LIABLE FOR ANY DIRECT,
-INDIRECT, INCIDENTAL, SPECIAL, EXEMPLARY, OR CONSEQUENTIAL DAMAGES (INCLUDING,
-BUT NOT LIMITED TO, PROCUREMENT OF SUBSTITUTE GOODS OR SERVICES; LOSS OF USE, DATA,
-OR PROFITS; OR BUSINESS INTERRUPTION) HOWEVER CAUSED AND ON ANY THEORY OF LIABILITY,
-WHETHER IN CONTRACT, STRICT LIABILITY, OR TORT (INCLUDING NEGLIGENCE OR OTHERWISE)
-ARISING IN ANY WAY OUT OF THE USE OF THIS SOFTWARE, EVEN IF ADVISED OF THE
-POSSIBILITY OF SUCH DAMAGE.
-*/
-
+// Copyright 1998-2019 Epic Games, Inc. All Rights Reserved.
 
 #pragma once
 
diff --git a/Engine/Plugins/Runtime/Steam/SteamVR/Source/SteamVRInputDevice/Public/ISteamVRInputDeviceModule.h b/Engine/Plugins/Runtime/Steam/SteamVR/Source/SteamVRInputDevice/Public/ISteamVRInputDeviceModule.h
--- a/Engine/Plugins/Runtime/Steam/SteamVR/Source/SteamVRInputDevice/Public/ISteamVRInputDeviceModule.h
+++ b/Engine/Plugins/Runtime/Steam/SteamVR/Source/SteamVRInputDevice/Public/ISteamVRInputDeviceModule.h
@@ -1,32 +1,4 @@
-/*
-Copyright 2019 Valve Corporation under https://opensource.org/licenses/BSD-3-Clause
-
-Redistribution and use in source and binary forms, with or without modification,
-are permitted provided that the following conditions are met:
-
-1. Redistributions of source code must retain the above copyright notice, this
-   list of conditions and the following disclaimer.
-
-2. Redistributions in binary form must reproduce the above copyright notice,
-   this list of conditions and the following disclaimer in the documentation
-   and/or other materials provided with the distribution.
-
-3. Neither the name of the copyright holder nor the names of its contributors
-   may be used to endorse or promote products derived from this software
-   without specific prior written permission.
-
-THIS SOFTWARE IS PROVIDED BY THE COPYRIGHT HOLDERS AND CONTRIBUTORS "AS IS" AND
-ANY EXPRESS OR IMPLIED WARRANTIES, INCLUDING, BUT NOT LIMITED TO, THE IMPLIED
-WARRANTIES OF MERCHANTABILITY AND FITNESS FOR A PARTICULAR PURPOSE ARE DISCLAIMED.
-IN NO EVENT SHALL THE COPYRIGHT HOLDER OR CONTRIBUTORS BE LIABLE FOR ANY DIRECT,
-INDIRECT, INCIDENTAL, SPECIAL, EXEMPLARY, OR CONSEQUENTIAL DAMAGES (INCLUDING,
-BUT NOT LIMITED TO, PROCUREMENT OF SUBSTITUTE GOODS OR SERVICES; LOSS OF USE, DATA,
-OR PROFITS; OR BUSINESS INTERRUPTION) HOWEVER CAUSED AND ON ANY THEORY OF LIABILITY,
-WHETHER IN CONTRACT, STRICT LIABILITY, OR TORT (INCLUDING NEGLIGENCE OR OTHERWISE)
-ARISING IN ANY WAY OUT OF THE USE OF THIS SOFTWARE, EVEN IF ADVISED OF THE
-POSSIBILITY OF SUCH DAMAGE.
-*/
-
+// Copyright 1998-2019 Epic Games, Inc. All Rights Reserved.
 
 #pragma once
 
diff --git a/Engine/Plugins/Runtime/Steam/SteamVR/Source/SteamVRInputDevice/Public/SteamVRControllerKeys.h b/Engine/Plugins/Runtime/Steam/SteamVR/Source/SteamVRInputDevice/Public/SteamVRControllerKeys.h
--- a/Engine/Plugins/Runtime/Steam/SteamVR/Source/SteamVRInputDevice/Public/SteamVRControllerKeys.h
+++ b/Engine/Plugins/Runtime/Steam/SteamVR/Source/SteamVRInputDevice/Public/SteamVRControllerKeys.h
@@ -1,32 +1,4 @@
-/*
-Copyright 2019 Valve Corporation under https://opensource.org/licenses/BSD-3-Clause
-
-Redistribution and use in source and binary forms, with or without modification,
-are permitted provided that the following conditions are met:
-
-1. Redistributions of source code must retain the above copyright notice, this
-   list of conditions and the following disclaimer.
-
-2. Redistributions in binary form must reproduce the above copyright notice,
-   this list of conditions and the following disclaimer in the documentation
-   and/or other materials provided with the distribution.
-
-3. Neither the name of the copyright holder nor the names of its contributors
-   may be used to endorse or promote products derived from this software
-   without specific prior written permission.
-
-THIS SOFTWARE IS PROVIDED BY THE COPYRIGHT HOLDERS AND CONTRIBUTORS "AS IS" AND
-ANY EXPRESS OR IMPLIED WARRANTIES, INCLUDING, BUT NOT LIMITED TO, THE IMPLIED
-WARRANTIES OF MERCHANTABILITY AND FITNESS FOR A PARTICULAR PURPOSE ARE DISCLAIMED.
-IN NO EVENT SHALL THE COPYRIGHT HOLDER OR CONTRIBUTORS BE LIABLE FOR ANY DIRECT,
-INDIRECT, INCIDENTAL, SPECIAL, EXEMPLARY, OR CONSEQUENTIAL DAMAGES (INCLUDING,
-BUT NOT LIMITED TO, PROCUREMENT OF SUBSTITUTE GOODS OR SERVICES; LOSS OF USE, DATA,
-OR PROFITS; OR BUSINESS INTERRUPTION) HOWEVER CAUSED AND ON ANY THEORY OF LIABILITY,
-WHETHER IN CONTRACT, STRICT LIABILITY, OR TORT (INCLUDING NEGLIGENCE OR OTHERWISE)
-ARISING IN ANY WAY OUT OF THE USE OF THIS SOFTWARE, EVEN IF ADVISED OF THE
-POSSIBILITY OF SUCH DAMAGE.
-*/
-
+// Copyright 1998-2019 Epic Games, Inc. All Rights Reserved.
 
 #pragma once
 #include "InputCoreTypes.h"
diff --git a/Engine/Plugins/Runtime/Steam/SteamVR/Source/SteamVRInputDevice/Public/SteamVRInputDevice.h b/Engine/Plugins/Runtime/Steam/SteamVR/Source/SteamVRInputDevice/Public/SteamVRInputDevice.h
--- a/Engine/Plugins/Runtime/Steam/SteamVR/Source/SteamVRInputDevice/Public/SteamVRInputDevice.h
+++ b/Engine/Plugins/Runtime/Steam/SteamVR/Source/SteamVRInputDevice/Public/SteamVRInputDevice.h
@@ -1,32 +1,4 @@
-/*
-Copyright 2019 Valve Corporation under https://opensource.org/licenses/BSD-3-Clause
-
-Redistribution and use in source and binary forms, with or without modification,
-are permitted provided that the following conditions are met:
-
-1. Redistributions of source code must retain the above copyright notice, this
-   list of conditions and the following disclaimer.
-
-2. Redistributions in binary form must reproduce the above copyright notice,
-   this list of conditions and the following disclaimer in the documentation
-   and/or other materials provided with the distribution.
-
-3. Neither the name of the copyright holder nor the names of its contributors
-   may be used to endorse or promote products derived from this software
-   without specific prior written permission.
-
-THIS SOFTWARE IS PROVIDED BY THE COPYRIGHT HOLDERS AND CONTRIBUTORS "AS IS" AND
-ANY EXPRESS OR IMPLIED WARRANTIES, INCLUDING, BUT NOT LIMITED TO, THE IMPLIED
-WARRANTIES OF MERCHANTABILITY AND FITNESS FOR A PARTICULAR PURPOSE ARE DISCLAIMED.
-IN NO EVENT SHALL THE COPYRIGHT HOLDER OR CONTRIBUTORS BE LIABLE FOR ANY DIRECT,
-INDIRECT, INCIDENTAL, SPECIAL, EXEMPLARY, OR CONSEQUENTIAL DAMAGES (INCLUDING,
-BUT NOT LIMITED TO, PROCUREMENT OF SUBSTITUTE GOODS OR SERVICES; LOSS OF USE, DATA,
-OR PROFITS; OR BUSINESS INTERRUPTION) HOWEVER CAUSED AND ON ANY THEORY OF LIABILITY,
-WHETHER IN CONTRACT, STRICT LIABILITY, OR TORT (INCLUDING NEGLIGENCE OR OTHERWISE)
-ARISING IN ANY WAY OUT OF THE USE OF THIS SOFTWARE, EVEN IF ADVISED OF THE
-POSSIBILITY OF SUCH DAMAGE.
-*/
-
+// Copyright 1998-2019 Epic Games, Inc. All Rights Reserved.
 
 #pragma once
 
diff --git a/Engine/Plugins/Runtime/Steam/SteamVR/Source/SteamVRInputDevice/Public/SteamVRInputDeviceFunctionLibrary.h b/Engine/Plugins/Runtime/Steam/SteamVR/Source/SteamVRInputDevice/Public/SteamVRInputDeviceFunctionLibrary.h
--- a/Engine/Plugins/Runtime/Steam/SteamVR/Source/SteamVRInputDevice/Public/SteamVRInputDeviceFunctionLibrary.h
+++ b/Engine/Plugins/Runtime/Steam/SteamVR/Source/SteamVRInputDevice/Public/SteamVRInputDeviceFunctionLibrary.h
@@ -1,32 +1,4 @@
-/*
-Copyright 2019 Valve Corporation under https://opensource.org/licenses/BSD-3-Clause
-
-Redistribution and use in source and binary forms, with or without modification,
-are permitted provided that the following conditions are met:
-
-1. Redistributions of source code must retain the above copyright notice, this
-   list of conditions and the following disclaimer.
-
-2. Redistributions in binary form must reproduce the above copyright notice,
-   this list of conditions and the following disclaimer in the documentation
-   and/or other materials provided with the distribution.
-
-3. Neither the name of the copyright holder nor the names of its contributors
-   may be used to endorse or promote products derived from this software
-   without specific prior written permission.
-
-THIS SOFTWARE IS PROVIDED BY THE COPYRIGHT HOLDERS AND CONTRIBUTORS "AS IS" AND
-ANY EXPRESS OR IMPLIED WARRANTIES, INCLUDING, BUT NOT LIMITED TO, THE IMPLIED
-WARRANTIES OF MERCHANTABILITY AND FITNESS FOR A PARTICULAR PURPOSE ARE DISCLAIMED.
-IN NO EVENT SHALL THE COPYRIGHT HOLDER OR CONTRIBUTORS BE LIABLE FOR ANY DIRECT,
-INDIRECT, INCIDENTAL, SPECIAL, EXEMPLARY, OR CONSEQUENTIAL DAMAGES (INCLUDING,
-BUT NOT LIMITED TO, PROCUREMENT OF SUBSTITUTE GOODS OR SERVICES; LOSS OF USE, DATA,
-OR PROFITS; OR BUSINESS INTERRUPTION) HOWEVER CAUSED AND ON ANY THEORY OF LIABILITY,
-WHETHER IN CONTRACT, STRICT LIABILITY, OR TORT (INCLUDING NEGLIGENCE OR OTHERWISE)
-ARISING IN ANY WAY OUT OF THE USE OF THIS SOFTWARE, EVEN IF ADVISED OF THE
-POSSIBILITY OF SUCH DAMAGE.
-*/
-
+// Copyright 1998-2019 Epic Games, Inc. All Rights Reserved.
 
 #pragma once
 
@@ -118,9 +90,6 @@ struct STEAMVRINPUTDEVICE_API FSteamVRSkeletonTransform
 	FTransform Aux_Ring;
 	UPROPERTY(EditAnywhere, BlueprintReadOnly, Category = "SteamVR Input")
 	FTransform Aux_Pinky;
-
-	UPROPERTY(EditAnywhere, BlueprintReadOnly, Category = "SteamVR Input")
-	FTransform Bone_Count;
 };
 
 /** SteamVR finger curls */
diff --git a/Engine/Plugins/Runtime/Steam/SteamVR/Source/SteamVRInputDevice/Public/SteamVRInputPublic.h b/Engine/Plugins/Runtime/Steam/SteamVR/Source/SteamVRInputDevice/Public/SteamVRInputPublic.h
--- a/Engine/Plugins/Runtime/Steam/SteamVR/Source/SteamVRInputDevice/Public/SteamVRInputPublic.h
+++ b/Engine/Plugins/Runtime/Steam/SteamVR/Source/SteamVRInputDevice/Public/SteamVRInputPublic.h
@@ -1,32 +1,4 @@
-/*
-Copyright 2019 Valve Corporation under https://opensource.org/licenses/BSD-3-Clause
-
-Redistribution and use in source and binary forms, with or without modification,
-are permitted provided that the following conditions are met:
-
-1. Redistributions of source code must retain the above copyright notice, this
-   list of conditions and the following disclaimer.
-
-2. Redistributions in binary form must reproduce the above copyright notice,
-   this list of conditions and the following disclaimer in the documentation
-   and/or other materials provided with the distribution.
-
-3. Neither the name of the copyright holder nor the names of its contributors
-   may be used to endorse or promote products derived from this software
-   without specific prior written permission.
-
-THIS SOFTWARE IS PROVIDED BY THE COPYRIGHT HOLDERS AND CONTRIBUTORS "AS IS" AND
-ANY EXPRESS OR IMPLIED WARRANTIES, INCLUDING, BUT NOT LIMITED TO, THE IMPLIED
-WARRANTIES OF MERCHANTABILITY AND FITNESS FOR A PARTICULAR PURPOSE ARE DISCLAIMED.
-IN NO EVENT SHALL THE COPYRIGHT HOLDER OR CONTRIBUTORS BE LIABLE FOR ANY DIRECT,
-INDIRECT, INCIDENTAL, SPECIAL, EXEMPLARY, OR CONSEQUENTIAL DAMAGES (INCLUDING,
-BUT NOT LIMITED TO, PROCUREMENT OF SUBSTITUTE GOODS OR SERVICES; LOSS OF USE, DATA,
-OR PROFITS; OR BUSINESS INTERRUPTION) HOWEVER CAUSED AND ON ANY THEORY OF LIABILITY,
-WHETHER IN CONTRACT, STRICT LIABILITY, OR TORT (INCLUDING NEGLIGENCE OR OTHERWISE)
-ARISING IN ANY WAY OUT OF THE USE OF THIS SOFTWARE, EVEN IF ADVISED OF THE
-POSSIBILITY OF SUCH DAMAGE.
-*/
-
+// Copyright 1998-2019 Epic Games, Inc. All Rights Reserved.
 
 struct FSteamVRSkeletonTransform_t
 {
diff --git a/Engine/Plugins/Runtime/Steam/SteamVR/Source/SteamVRInputDevice/Public/SteamVRInputTypes.h b/Engine/Plugins/Runtime/Steam/SteamVR/Source/SteamVRInputDevice/Public/SteamVRInputTypes.h
--- a/Engine/Plugins/Runtime/Steam/SteamVR/Source/SteamVRInputDevice/Public/SteamVRInputTypes.h
+++ b/Engine/Plugins/Runtime/Steam/SteamVR/Source/SteamVRInputDevice/Public/SteamVRInputTypes.h
@@ -1,32 +1,4 @@
-/*
-Copyright 2019 Valve Corporation under https://opensource.org/licenses/BSD-3-Clause
-
-Redistribution and use in source and binary forms, with or without modification,
-are permitted provided that the following conditions are met:
-
-1. Redistributions of source code must retain the above copyright notice, this
-   list of conditions and the following disclaimer.
-
-2. Redistributions in binary form must reproduce the above copyright notice,
-   this list of conditions and the following disclaimer in the documentation
-   and/or other materials provided with the distribution.
-
-3. Neither the name of the copyright holder nor the names of its contributors
-   may be used to endorse or promote products derived from this software
-   without specific prior written permission.
-
-THIS SOFTWARE IS PROVIDED BY THE COPYRIGHT HOLDERS AND CONTRIBUTORS "AS IS" AND
-ANY EXPRESS OR IMPLIED WARRANTIES, INCLUDING, BUT NOT LIMITED TO, THE IMPLIED
-WARRANTIES OF MERCHANTABILITY AND FITNESS FOR A PARTICULAR PURPOSE ARE DISCLAIMED.
-IN NO EVENT SHALL THE COPYRIGHT HOLDER OR CONTRIBUTORS BE LIABLE FOR ANY DIRECT,
-INDIRECT, INCIDENTAL, SPECIAL, EXEMPLARY, OR CONSEQUENTIAL DAMAGES (INCLUDING,
-BUT NOT LIMITED TO, PROCUREMENT OF SUBSTITUTE GOODS OR SERVICES; LOSS OF USE, DATA,
-OR PROFITS; OR BUSINESS INTERRUPTION) HOWEVER CAUSED AND ON ANY THEORY OF LIABILITY,
-WHETHER IN CONTRACT, STRICT LIABILITY, OR TORT (INCLUDING NEGLIGENCE OR OTHERWISE)
-ARISING IN ANY WAY OUT OF THE USE OF THIS SOFTWARE, EVEN IF ADVISED OF THE
-POSSIBILITY OF SUCH DAMAGE.
-*/
-
+// Copyright 1998-2019 Epic Games, Inc. All Rights Reserved.
 
 #pragma once
 
diff --git a/Engine/Plugins/Runtime/Steam/SteamVR/Source/SteamVRInputDevice/Public/SteamVRKnucklesKeys.h b/Engine/Plugins/Runtime/Steam/SteamVR/Source/SteamVRInputDevice/Public/SteamVRKnucklesKeys.h
--- a/Engine/Plugins/Runtime/Steam/SteamVR/Source/SteamVRInputDevice/Public/SteamVRKnucklesKeys.h
+++ b/Engine/Plugins/Runtime/Steam/SteamVR/Source/SteamVRInputDevice/Public/SteamVRKnucklesKeys.h
@@ -1,32 +1,4 @@
-/*
-Copyright 2019 Valve Corporation under https://opensource.org/licenses/BSD-3-Clause
-
-Redistribution and use in source and binary forms, with or without modification,
-are permitted provided that the following conditions are met:
-
-1. Redistributions of source code must retain the above copyright notice, this
-   list of conditions and the following disclaimer.
-
-2. Redistributions in binary form must reproduce the above copyright notice,
-   this list of conditions and the following disclaimer in the documentation
-   and/or other materials provided with the distribution.
-
-3. Neither the name of the copyright holder nor the names of its contributors
-   may be used to endorse or promote products derived from this software
-   without specific prior written permission.
-
-THIS SOFTWARE IS PROVIDED BY THE COPYRIGHT HOLDERS AND CONTRIBUTORS "AS IS" AND
-ANY EXPRESS OR IMPLIED WARRANTIES, INCLUDING, BUT NOT LIMITED TO, THE IMPLIED
-WARRANTIES OF MERCHANTABILITY AND FITNESS FOR A PARTICULAR PURPOSE ARE DISCLAIMED.
-IN NO EVENT SHALL THE COPYRIGHT HOLDER OR CONTRIBUTORS BE LIABLE FOR ANY DIRECT,
-INDIRECT, INCIDENTAL, SPECIAL, EXEMPLARY, OR CONSEQUENTIAL DAMAGES (INCLUDING,
-BUT NOT LIMITED TO, PROCUREMENT OF SUBSTITUTE GOODS OR SERVICES; LOSS OF USE, DATA,
-OR PROFITS; OR BUSINESS INTERRUPTION) HOWEVER CAUSED AND ON ANY THEORY OF LIABILITY,
-WHETHER IN CONTRACT, STRICT LIABILITY, OR TORT (INCLUDING NEGLIGENCE OR OTHERWISE)
-ARISING IN ANY WAY OUT OF THE USE OF THIS SOFTWARE, EVEN IF ADVISED OF THE
-POSSIBILITY OF SUCH DAMAGE.
-*/
-
+// Copyright 1998-2019 Epic Games, Inc. All Rights Reserved.
 
 #pragma once
 #include "InputCoreTypes.h"
diff --git a/Engine/Plugins/Runtime/Steam/SteamVR/Source/SteamVRInputDevice/Public/SteamVRSkeletonDefinition.h b/Engine/Plugins/Runtime/Steam/SteamVR/Source/SteamVRInputDevice/Public/SteamVRSkeletonDefinition.h
--- a/Engine/Plugins/Runtime/Steam/SteamVR/Source/SteamVRInputDevice/Public/SteamVRSkeletonDefinition.h
+++ b/Engine/Plugins/Runtime/Steam/SteamVR/Source/SteamVRInputDevice/Public/SteamVRSkeletonDefinition.h
@@ -1,32 +1,4 @@
-/*
-Copyright 2019 Valve Corporation under https://opensource.org/licenses/BSD-3-Clause
-
-Redistribution and use in source and binary forms, with or without modification,
-are permitted provided that the following conditions are met:
-
-1. Redistributions of source code must retain the above copyright notice, this
-   list of conditions and the following disclaimer.
-
-2. Redistributions in binary form must reproduce the above copyright notice,
-   this list of conditions and the following disclaimer in the documentation
-   and/or other materials provided with the distribution.
-
-3. Neither the name of the copyright holder nor the names of its contributors
-   may be used to endorse or promote products derived from this software
-   without specific prior written permission.
-
-THIS SOFTWARE IS PROVIDED BY THE COPYRIGHT HOLDERS AND CONTRIBUTORS "AS IS" AND
-ANY EXPRESS OR IMPLIED WARRANTIES, INCLUDING, BUT NOT LIMITED TO, THE IMPLIED
-WARRANTIES OF MERCHANTABILITY AND FITNESS FOR A PARTICULAR PURPOSE ARE DISCLAIMED.
-IN NO EVENT SHALL THE COPYRIGHT HOLDER OR CONTRIBUTORS BE LIABLE FOR ANY DIRECT,
-INDIRECT, INCIDENTAL, SPECIAL, EXEMPLARY, OR CONSEQUENTIAL DAMAGES (INCLUDING,
-BUT NOT LIMITED TO, PROCUREMENT OF SUBSTITUTE GOODS OR SERVICES; LOSS OF USE, DATA,
-OR PROFITS; OR BUSINESS INTERRUPTION) HOWEVER CAUSED AND ON ANY THEORY OF LIABILITY,
-WHETHER IN CONTRACT, STRICT LIABILITY, OR TORT (INCLUDING NEGLIGENCE OR OTHERWISE)
-ARISING IN ANY WAY OUT OF THE USE OF THIS SOFTWARE, EVEN IF ADVISED OF THE
-POSSIBILITY OF SUCH DAMAGE.
-*/
-
+// Copyright 1998-2019 Epic Games, Inc. All Rights Reserved.
 
 #pragma once
 
diff --git a/Engine/Plugins/Runtime/Steam/SteamVR/Source/SteamVRInputDevice/Public/SteamVRTrackingRefComponent.h b/Engine/Plugins/Runtime/Steam/SteamVR/Source/SteamVRInputDevice/Public/SteamVRTrackingRefComponent.h
--- a/Engine/Plugins/Runtime/Steam/SteamVR/Source/SteamVRInputDevice/Public/SteamVRTrackingRefComponent.h
+++ b/Engine/Plugins/Runtime/Steam/SteamVR/Source/SteamVRInputDevice/Public/SteamVRTrackingRefComponent.h
@@ -1,31 +1,4 @@
-/*
-Copyright 2019 Valve Corporation under https://opensource.org/licenses/BSD-3-Clause
-
-Redistribution and use in source and binary forms, with or without modification,
-are permitted provided that the following conditions are met:
-
-1. Redistributions of source code must retain the above copyright notice, this
-   list of conditions and the following disclaimer.
-
-2. Redistributions in binary form must reproduce the above copyright notice,
-   this list of conditions and the following disclaimer in the documentation
-   and/or other materials provided with the distribution.
-
-3. Neither the name of the copyright holder nor the names of its contributors
-   may be used to endorse or promote products derived from this software
-   without specific prior written permission.
-
-THIS SOFTWARE IS PROVIDED BY THE COPYRIGHT HOLDERS AND CONTRIBUTORS "AS IS" AND
-ANY EXPRESS OR IMPLIED WARRANTIES, INCLUDING, BUT NOT LIMITED TO, THE IMPLIED
-WARRANTIES OF MERCHANTABILITY AND FITNESS FOR A PARTICULAR PURPOSE ARE DISCLAIMED.
-IN NO EVENT SHALL THE COPYRIGHT HOLDER OR CONTRIBUTORS BE LIABLE FOR ANY DIRECT,
-INDIRECT, INCIDENTAL, SPECIAL, EXEMPLARY, OR CONSEQUENTIAL DAMAGES (INCLUDING,
-BUT NOT LIMITED TO, PROCUREMENT OF SUBSTITUTE GOODS OR SERVICES; LOSS OF USE, DATA,
-OR PROFITS; OR BUSINESS INTERRUPTION) HOWEVER CAUSED AND ON ANY THEORY OF LIABILITY,
-WHETHER IN CONTRACT, STRICT LIABILITY, OR TORT (INCLUDING NEGLIGENCE OR OTHERWISE)
-ARISING IN ANY WAY OUT OF THE USE OF THIS SOFTWARE, EVEN IF ADVISED OF THE
-POSSIBILITY OF SUCH DAMAGE.
-*/
+// Copyright 1998-2019 Epic Games, Inc. All Rights Reserved.
 
 #pragma once
 
diff --git a/Engine/Plugins/Runtime/Steam/SteamVR/Source/SteamVRInputDevice/Public/UE4HandSkeletonDefinition.h b/Engine/Plugins/Runtime/Steam/SteamVR/Source/SteamVRInputDevice/Public/UE4HandSkeletonDefinition.h
--- a/Engine/Plugins/Runtime/Steam/SteamVR/Source/SteamVRInputDevice/Public/UE4HandSkeletonDefinition.h
+++ b/Engine/Plugins/Runtime/Steam/SteamVR/Source/SteamVRInputDevice/Public/UE4HandSkeletonDefinition.h
@@ -1,32 +1,4 @@
-/*
-Copyright 2019 Valve Corporation under https://opensource.org/licenses/BSD-3-Clause
-
-Redistribution and use in source and binary forms, with or without modification,
-are permitted provided that the following conditions are met:
-
-1. Redistributions of source code must retain the above copyright notice, this
-   list of conditions and the following disclaimer.
-
-2. Redistributions in binary form must reproduce the above copyright notice,
-   this list of conditions and the following disclaimer in the documentation
-   and/or other materials provided with the distribution.
-
-3. Neither the name of the copyright holder nor the names of its contributors
-   may be used to endorse or promote products derived from this software
-   without specific prior written permission.
-
-THIS SOFTWARE IS PROVIDED BY THE COPYRIGHT HOLDERS AND CONTRIBUTORS "AS IS" AND
-ANY EXPRESS OR IMPLIED WARRANTIES, INCLUDING, BUT NOT LIMITED TO, THE IMPLIED
-WARRANTIES OF MERCHANTABILITY AND FITNESS FOR A PARTICULAR PURPOSE ARE DISCLAIMED.
-IN NO EVENT SHALL THE COPYRIGHT HOLDER OR CONTRIBUTORS BE LIABLE FOR ANY DIRECT,
-INDIRECT, INCIDENTAL, SPECIAL, EXEMPLARY, OR CONSEQUENTIAL DAMAGES (INCLUDING,
-BUT NOT LIMITED TO, PROCUREMENT OF SUBSTITUTE GOODS OR SERVICES; LOSS OF USE, DATA,
-OR PROFITS; OR BUSINESS INTERRUPTION) HOWEVER CAUSED AND ON ANY THEORY OF LIABILITY,
-WHETHER IN CONTRACT, STRICT LIABILITY, OR TORT (INCLUDING NEGLIGENCE OR OTHERWISE)
-ARISING IN ANY WAY OUT OF THE USE OF THIS SOFTWARE, EVEN IF ADVISED OF THE
-POSSIBILITY OF SUCH DAMAGE.
-*/
-
+// Copyright 1998-2019 Epic Games, Inc. All Rights Reserved.
 
 #pragma once
 
