More static analysis fixes for SteamVR Input and copyright header updates

diff --git a/Engine/Plugins/Runtime/Steam/SteamVR/Source/SteamVREditor/Private/AnimGraphNode_SteamVRInputAnimPose.cpp b/Engine/Plugins/Runtime/Steam/SteamVR/Source/SteamVREditor/Private/AnimGraphNode_SteamVRInputAnimPose.cpp
--- a/Engine/Plugins/Runtime/Steam/SteamVR/Source/SteamVREditor/Private/AnimGraphNode_SteamVRInputAnimPose.cpp
+++ b/Engine/Plugins/Runtime/Steam/SteamVR/Source/SteamVREditor/Private/AnimGraphNode_SteamVRInputAnimPose.cpp
@@ -1,32 +1,4 @@
-/*
-Copyright 2019 Valve Corporation under https://opensource.org/licenses/BSD-3-Clause
-
-Redistribution and use in source and binary forms, with or without modification,
-are permitted provided that the following conditions are met:
-
-1. Redistributions of source code must retain the above copyright notice, this
-   list of conditions and the following disclaimer.
-
-2. Redistributions in binary form must reproduce the above copyright notice,
-   this list of conditions and the following disclaimer in the documentation
-   and/or other materials provided with the distribution.
-
-3. Neither the name of the copyright holder nor the names of its contributors
-   may be used to endorse or promote products derived from this software
-   without specific prior written permission.
-
-THIS SOFTWARE IS PROVIDED BY THE COPYRIGHT HOLDERS AND CONTRIBUTORS "AS IS" AND
-ANY EXPRESS OR IMPLIED WARRANTIES, INCLUDING, BUT NOT LIMITED TO, THE IMPLIED
-WARRANTIES OF MERCHANTABILITY AND FITNESS FOR A PARTICULAR PURPOSE ARE DISCLAIMED.
-IN NO EVENT SHALL THE COPYRIGHT HOLDER OR CONTRIBUTORS BE LIABLE FOR ANY DIRECT,
-INDIRECT, INCIDENTAL, SPECIAL, EXEMPLARY, OR CONSEQUENTIAL DAMAGES (INCLUDING,
-BUT NOT LIMITED TO, PROCUREMENT OF SUBSTITUTE GOODS OR SERVICES; LOSS OF USE, DATA,
-OR PROFITS; OR BUSINESS INTERRUPTION) HOWEVER CAUSED AND ON ANY THEORY OF LIABILITY,
-WHETHER IN CONTRACT, STRICT LIABILITY, OR TORT (INCLUDING NEGLIGENCE OR OTHERWISE)
-ARISING IN ANY WAY OUT OF THE USE OF THIS SOFTWARE, EVEN IF ADVISED OF THE
-POSSIBILITY OF SUCH DAMAGE.
-*/
-
+// Copyright 1998-2019 Epic Games, Inc. All Rights Reserved.
 
 #include "AnimGraphNode_SteamVRInputAnimPose.h"
 #include "SteamVREditor.h"
diff --git a/Engine/Plugins/Runtime/Steam/SteamVR/Source/SteamVREditor/Private/AnimGraphNode_SteamVRSetWristTransform.cpp b/Engine/Plugins/Runtime/Steam/SteamVR/Source/SteamVREditor/Private/AnimGraphNode_SteamVRSetWristTransform.cpp
--- a/Engine/Plugins/Runtime/Steam/SteamVR/Source/SteamVREditor/Private/AnimGraphNode_SteamVRSetWristTransform.cpp
+++ b/Engine/Plugins/Runtime/Steam/SteamVR/Source/SteamVREditor/Private/AnimGraphNode_SteamVRSetWristTransform.cpp
@@ -1,31 +1,4 @@
-/*
-Copyright 2019 Valve Corporation under https://opensource.org/licenses/BSD-3-Clause
-
-Redistribution and use in source and binary forms, with or without modification,
-are permitted provided that the following conditions are met:
-
-1. Redistributions of source code must retain the above copyright notice, this
-   list of conditions and the following disclaimer.
-
-2. Redistributions in binary form must reproduce the above copyright notice,
-   this list of conditions and the following disclaimer in the documentation
-   and/or other materials provided with the distribution.
-
-3. Neither the name of the copyright holder nor the names of its contributors
-   may be used to endorse or promote products derived from this software
-   without specific prior written permission.
-
-THIS SOFTWARE IS PROVIDED BY THE COPYRIGHT HOLDERS AND CONTRIBUTORS "AS IS" AND
-ANY EXPRESS OR IMPLIED WARRANTIES, INCLUDING, BUT NOT LIMITED TO, THE IMPLIED
-WARRANTIES OF MERCHANTABILITY AND FITNESS FOR A PARTICULAR PURPOSE ARE DISCLAIMED.
-IN NO EVENT SHALL THE COPYRIGHT HOLDER OR CONTRIBUTORS BE LIABLE FOR ANY DIRECT,
-INDIRECT, INCIDENTAL, SPECIAL, EXEMPLARY, OR CONSEQUENTIAL DAMAGES (INCLUDING,
-BUT NOT LIMITED TO, PROCUREMENT OF SUBSTITUTE GOODS OR SERVICES; LOSS OF USE, DATA,
-OR PROFITS; OR BUSINESS INTERRUPTION) HOWEVER CAUSED AND ON ANY THEORY OF LIABILITY,
-WHETHER IN CONTRACT, STRICT LIABILITY, OR TORT (INCLUDING NEGLIGENCE OR OTHERWISE)
-ARISING IN ANY WAY OUT OF THE USE OF THIS SOFTWARE, EVEN IF ADVISED OF THE
-POSSIBILITY OF SUCH DAMAGE.
-*/
+// Copyright 1998-2019 Epic Games, Inc. All Rights Reserved.
 
 #include "AnimGraphNode_SteamVRSetWristTransform.h"
 #include "SteamVREditor.h"
diff --git a/Engine/Plugins/Runtime/Steam/SteamVR/Source/SteamVREditor/Private/SteamVREditor.cpp b/Engine/Plugins/Runtime/Steam/SteamVR/Source/SteamVREditor/Private/SteamVREditor.cpp
--- a/Engine/Plugins/Runtime/Steam/SteamVR/Source/SteamVREditor/Private/SteamVREditor.cpp
+++ b/Engine/Plugins/Runtime/Steam/SteamVR/Source/SteamVREditor/Private/SteamVREditor.cpp
@@ -1,31 +1,4 @@
-/*
-Copyright 2019 Valve Corporation under https://opensource.org/licenses/BSD-3-Clause
-
-Redistribution and use in source and binary forms, with or without modification,
-are permitted provided that the following conditions are met:
-
-1. Redistributions of source code must retain the above copyright notice, this 
-   list of conditions and the following disclaimer.
-
-2. Redistributions in binary form must reproduce the above copyright notice, 
-   this list of conditions and the following disclaimer in the documentation 
-   and/or other materials provided with the distribution.
-
-3. Neither the name of the copyright holder nor the names of its contributors 
-   may be used to endorse or promote products derived from this software 
-   without specific prior written permission.
-
-THIS SOFTWARE IS PROVIDED BY THE COPYRIGHT HOLDERS AND CONTRIBUTORS "AS IS" AND
-ANY EXPRESS OR IMPLIED WARRANTIES, INCLUDING, BUT NOT LIMITED TO, THE IMPLIED 
-WARRANTIES OF MERCHANTABILITY AND FITNESS FOR A PARTICULAR PURPOSE ARE DISCLAIMED. 
-IN NO EVENT SHALL THE COPYRIGHT HOLDER OR CONTRIBUTORS BE LIABLE FOR ANY DIRECT, 
-INDIRECT, INCIDENTAL, SPECIAL, EXEMPLARY, OR CONSEQUENTIAL DAMAGES (INCLUDING, 
-BUT NOT LIMITED TO, PROCUREMENT OF SUBSTITUTE GOODS OR SERVICES; LOSS OF USE, DATA, 
-OR PROFITS; OR BUSINESS INTERRUPTION) HOWEVER CAUSED AND ON ANY THEORY OF LIABILITY, 
-WHETHER IN CONTRACT, STRICT LIABILITY, OR TORT (INCLUDING NEGLIGENCE OR OTHERWISE) 
-ARISING IN ANY WAY OUT OF THE USE OF THIS SOFTWARE, EVEN IF ADVISED OF THE 
-POSSIBILITY OF SUCH DAMAGE.
-*/
+// Copyright 1998-2019 Epic Games, Inc. All Rights Reserved.
 
 #pragma once
 
diff --git a/Engine/Plugins/Runtime/Steam/SteamVR/Source/SteamVREditor/Private/SteamVREditorStyle.cpp b/Engine/Plugins/Runtime/Steam/SteamVR/Source/SteamVREditor/Private/SteamVREditorStyle.cpp
--- a/Engine/Plugins/Runtime/Steam/SteamVR/Source/SteamVREditor/Private/SteamVREditorStyle.cpp
+++ b/Engine/Plugins/Runtime/Steam/SteamVR/Source/SteamVREditor/Private/SteamVREditorStyle.cpp
@@ -1,31 +1,4 @@
-/*
-Copyright 2019 Valve Corporation under https://opensource.org/licenses/BSD-3-Clause
-
-Redistribution and use in source and binary forms, with or without modification,
-are permitted provided that the following conditions are met:
-
-1. Redistributions of source code must retain the above copyright notice, this
-   list of conditions and the following disclaimer.
-
-2. Redistributions in binary form must reproduce the above copyright notice,
-   this list of conditions and the following disclaimer in the documentation
-   and/or other materials provided with the distribution.
-
-3. Neither the name of the copyright holder nor the names of its contributors
-   may be used to endorse or promote products derived from this software
-   without specific prior written permission.
-
-THIS SOFTWARE IS PROVIDED BY THE COPYRIGHT HOLDERS AND CONTRIBUTORS "AS IS" AND
-ANY EXPRESS OR IMPLIED WARRANTIES, INCLUDING, BUT NOT LIMITED TO, THE IMPLIED
-WARRANTIES OF MERCHANTABILITY AND FITNESS FOR A PARTICULAR PURPOSE ARE DISCLAIMED.
-IN NO EVENT SHALL THE COPYRIGHT HOLDER OR CONTRIBUTORS BE LIABLE FOR ANY DIRECT,
-INDIRECT, INCIDENTAL, SPECIAL, EXEMPLARY, OR CONSEQUENTIAL DAMAGES (INCLUDING,
-BUT NOT LIMITED TO, PROCUREMENT OF SUBSTITUTE GOODS OR SERVICES; LOSS OF USE, DATA,
-OR PROFITS; OR BUSINESS INTERRUPTION) HOWEVER CAUSED AND ON ANY THEORY OF LIABILITY,
-WHETHER IN CONTRACT, STRICT LIABILITY, OR TORT (INCLUDING NEGLIGENCE OR OTHERWISE)
-ARISING IN ANY WAY OUT OF THE USE OF THIS SOFTWARE, EVEN IF ADVISED OF THE
-POSSIBILITY OF SUCH DAMAGE.
-*/
+// Copyright 1998-2019 Epic Games, Inc. All Rights Reserved.
 
 #include "SteamVREditorStyle.h"
 #include "SteamVREditor.h"
diff --git a/Engine/Plugins/Runtime/Steam/SteamVR/Source/SteamVREditor/Public/AnimGraphNode_SteamVRInputAnimPose.h b/Engine/Plugins/Runtime/Steam/SteamVR/Source/SteamVREditor/Public/AnimGraphNode_SteamVRInputAnimPose.h
--- a/Engine/Plugins/Runtime/Steam/SteamVR/Source/SteamVREditor/Public/AnimGraphNode_SteamVRInputAnimPose.h
+++ b/Engine/Plugins/Runtime/Steam/SteamVR/Source/SteamVREditor/Public/AnimGraphNode_SteamVRInputAnimPose.h
@@ -1,32 +1,4 @@
-/*
-Copyright 2019 Valve Corporation under https://opensource.org/licenses/BSD-3-Clause
-
-Redistribution and use in source and binary forms, with or without modification,
-are permitted provided that the following conditions are met:
-
-1. Redistributions of source code must retain the above copyright notice, this
-   list of conditions and the following disclaimer.
-
-2. Redistributions in binary form must reproduce the above copyright notice,
-   this list of conditions and the following disclaimer in the documentation
-   and/or other materials provided with the distribution.
-
-3. Neither the name of the copyright holder nor the names of its contributors
-   may be used to endorse or promote products derived from this software
-   without specific prior written permission.
-
-THIS SOFTWARE IS PROVIDED BY THE COPYRIGHT HOLDERS AND CONTRIBUTORS "AS IS" AND
-ANY EXPRESS OR IMPLIED WARRANTIES, INCLUDING, BUT NOT LIMITED TO, THE IMPLIED
-WARRANTIES OF MERCHANTABILITY AND FITNESS FOR A PARTICULAR PURPOSE ARE DISCLAIMED.
-IN NO EVENT SHALL THE COPYRIGHT HOLDER OR CONTRIBUTORS BE LIABLE FOR ANY DIRECT,
-INDIRECT, INCIDENTAL, SPECIAL, EXEMPLARY, OR CONSEQUENTIAL DAMAGES (INCLUDING,
-BUT NOT LIMITED TO, PROCUREMENT OF SUBSTITUTE GOODS OR SERVICES; LOSS OF USE, DATA,
-OR PROFITS; OR BUSINESS INTERRUPTION) HOWEVER CAUSED AND ON ANY THEORY OF LIABILITY,
-WHETHER IN CONTRACT, STRICT LIABILITY, OR TORT (INCLUDING NEGLIGENCE OR OTHERWISE)
-ARISING IN ANY WAY OUT OF THE USE OF THIS SOFTWARE, EVEN IF ADVISED OF THE
-POSSIBILITY OF SUCH DAMAGE.
-*/
-
+// Copyright 1998-2019 Epic Games, Inc. All Rights Reserved.
 
 #pragma once
 
diff --git a/Engine/Plugins/Runtime/Steam/SteamVR/Source/SteamVREditor/Public/AnimGraphNode_SteamVRSetWristTransform.h b/Engine/Plugins/Runtime/Steam/SteamVR/Source/SteamVREditor/Public/AnimGraphNode_SteamVRSetWristTransform.h
--- a/Engine/Plugins/Runtime/Steam/SteamVR/Source/SteamVREditor/Public/AnimGraphNode_SteamVRSetWristTransform.h
+++ b/Engine/Plugins/Runtime/Steam/SteamVR/Source/SteamVREditor/Public/AnimGraphNode_SteamVRSetWristTransform.h
@@ -1,32 +1,4 @@
-/*
-Copyright 2019 Valve Corporation under https://opensource.org/licenses/BSD-3-Clause
-
-Redistribution and use in source and binary forms, with or without modification,
-are permitted provided that the following conditions are met:
-
-1. Redistributions of source code must retain the above copyright notice, this
-   list of conditions and the following disclaimer.
-
-2. Redistributions in binary form must reproduce the above copyright notice,
-   this list of conditions and the following disclaimer in the documentation
-   and/or other materials provided with the distribution.
-
-3. Neither the name of the copyright holder nor the names of its contributors
-   may be used to endorse or promote products derived from this software
-   without specific prior written permission.
-
-THIS SOFTWARE IS PROVIDED BY THE COPYRIGHT HOLDERS AND CONTRIBUTORS "AS IS" AND
-ANY EXPRESS OR IMPLIED WARRANTIES, INCLUDING, BUT NOT LIMITED TO, THE IMPLIED
-WARRANTIES OF MERCHANTABILITY AND FITNESS FOR A PARTICULAR PURPOSE ARE DISCLAIMED.
-IN NO EVENT SHALL THE COPYRIGHT HOLDER OR CONTRIBUTORS BE LIABLE FOR ANY DIRECT,
-INDIRECT, INCIDENTAL, SPECIAL, EXEMPLARY, OR CONSEQUENTIAL DAMAGES (INCLUDING,
-BUT NOT LIMITED TO, PROCUREMENT OF SUBSTITUTE GOODS OR SERVICES; LOSS OF USE, DATA,
-OR PROFITS; OR BUSINESS INTERRUPTION) HOWEVER CAUSED AND ON ANY THEORY OF LIABILITY,
-WHETHER IN CONTRACT, STRICT LIABILITY, OR TORT (INCLUDING NEGLIGENCE OR OTHERWISE)
-ARISING IN ANY WAY OUT OF THE USE OF THIS SOFTWARE, EVEN IF ADVISED OF THE
-POSSIBILITY OF SUCH DAMAGE.
-*/
-
+// Copyright 1998-2019 Epic Games, Inc. All Rights Reserved.
 
 #pragma once
 
diff --git a/Engine/Plugins/Runtime/Steam/SteamVR/Source/SteamVREditor/Public/SteamVREditor.h b/Engine/Plugins/Runtime/Steam/SteamVR/Source/SteamVREditor/Public/SteamVREditor.h
--- a/Engine/Plugins/Runtime/Steam/SteamVR/Source/SteamVREditor/Public/SteamVREditor.h
+++ b/Engine/Plugins/Runtime/Steam/SteamVR/Source/SteamVREditor/Public/SteamVREditor.h
@@ -1,32 +1,4 @@
-/*
-Copyright 2019 Valve Corporation under https://opensource.org/licenses/BSD-3-Clause
-
-Redistribution and use in source and binary forms, with or without modification,
-are permitted provided that the following conditions are met:
-
-1. Redistributions of source code must retain the above copyright notice, this
-   list of conditions and the following disclaimer.
-
-2. Redistributions in binary form must reproduce the above copyright notice,
-   this list of conditions and the following disclaimer in the documentation
-   and/or other materials provided with the distribution.
-
-3. Neither the name of the copyright holder nor the names of its contributors
-   may be used to endorse or promote products derived from this software
-   without specific prior written permission.
-
-THIS SOFTWARE IS PROVIDED BY THE COPYRIGHT HOLDERS AND CONTRIBUTORS "AS IS" AND
-ANY EXPRESS OR IMPLIED WARRANTIES, INCLUDING, BUT NOT LIMITED TO, THE IMPLIED
-WARRANTIES OF MERCHANTABILITY AND FITNESS FOR A PARTICULAR PURPOSE ARE DISCLAIMED.
-IN NO EVENT SHALL THE COPYRIGHT HOLDER OR CONTRIBUTORS BE LIABLE FOR ANY DIRECT,
-INDIRECT, INCIDENTAL, SPECIAL, EXEMPLARY, OR CONSEQUENTIAL DAMAGES (INCLUDING,
-BUT NOT LIMITED TO, PROCUREMENT OF SUBSTITUTE GOODS OR SERVICES; LOSS OF USE, DATA,
-OR PROFITS; OR BUSINESS INTERRUPTION) HOWEVER CAUSED AND ON ANY THEORY OF LIABILITY,
-WHETHER IN CONTRACT, STRICT LIABILITY, OR TORT (INCLUDING NEGLIGENCE OR OTHERWISE)
-ARISING IN ANY WAY OUT OF THE USE OF THIS SOFTWARE, EVEN IF ADVISED OF THE
-POSSIBILITY OF SUCH DAMAGE.
-*/
-
+// Copyright 1998-2019 Epic Games, Inc. All Rights Reserved.
 
 #pragma once
 
diff --git a/Engine/Plugins/Runtime/Steam/SteamVR/Source/SteamVREditor/Public/SteamVREditorCommands.cpp b/Engine/Plugins/Runtime/Steam/SteamVR/Source/SteamVREditor/Public/SteamVREditorCommands.cpp
--- a/Engine/Plugins/Runtime/Steam/SteamVR/Source/SteamVREditor/Public/SteamVREditorCommands.cpp
+++ b/Engine/Plugins/Runtime/Steam/SteamVR/Source/SteamVREditor/Public/SteamVREditorCommands.cpp
@@ -1,31 +1,4 @@
-/*
-Copyright 2019 Valve Corporation under https://opensource.org/licenses/BSD-3-Clause
-
-Redistribution and use in source and binary forms, with or without modification,
-are permitted provided that the following conditions are met:
-
-1. Redistributions of source code must retain the above copyright notice, this
-   list of conditions and the following disclaimer.
-
-2. Redistributions in binary form must reproduce the above copyright notice,
-   this list of conditions and the following disclaimer in the documentation
-   and/or other materials provided with the distribution.
-
-3. Neither the name of the copyright holder nor the names of its contributors
-   may be used to endorse or promote products derived from this software
-   without specific prior written permission.
-
-THIS SOFTWARE IS PROVIDED BY THE COPYRIGHT HOLDERS AND CONTRIBUTORS "AS IS" AND
-ANY EXPRESS OR IMPLIED WARRANTIES, INCLUDING, BUT NOT LIMITED TO, THE IMPLIED
-WARRANTIES OF MERCHANTABILITY AND FITNESS FOR A PARTICULAR PURPOSE ARE DISCLAIMED.
-IN NO EVENT SHALL THE COPYRIGHT HOLDER OR CONTRIBUTORS BE LIABLE FOR ANY DIRECT,
-INDIRECT, INCIDENTAL, SPECIAL, EXEMPLARY, OR CONSEQUENTIAL DAMAGES (INCLUDING,
-BUT NOT LIMITED TO, PROCUREMENT OF SUBSTITUTE GOODS OR SERVICES; LOSS OF USE, DATA,
-OR PROFITS; OR BUSINESS INTERRUPTION) HOWEVER CAUSED AND ON ANY THEORY OF LIABILITY,
-WHETHER IN CONTRACT, STRICT LIABILITY, OR TORT (INCLUDING NEGLIGENCE OR OTHERWISE)
-ARISING IN ANY WAY OUT OF THE USE OF THIS SOFTWARE, EVEN IF ADVISED OF THE
-POSSIBILITY OF SUCH DAMAGE.
-*/
+// Copyright 1998-2019 Epic Games, Inc. All Rights Reserved.
 
 #include "SteamVREditorCommands.h"
 #include "SteamVREditor.h"
diff --git a/Engine/Plugins/Runtime/Steam/SteamVR/Source/SteamVREditor/Public/SteamVREditorCommands.h b/Engine/Plugins/Runtime/Steam/SteamVR/Source/SteamVREditor/Public/SteamVREditorCommands.h
--- a/Engine/Plugins/Runtime/Steam/SteamVR/Source/SteamVREditor/Public/SteamVREditorCommands.h
+++ b/Engine/Plugins/Runtime/Steam/SteamVR/Source/SteamVREditor/Public/SteamVREditorCommands.h
@@ -1,32 +1,4 @@
-/*
-Copyright 2019 Valve Corporation under https://opensource.org/licenses/BSD-3-Clause
-
-Redistribution and use in source and binary forms, with or without modification,
-are permitted provided that the following conditions are met:
-
-1. Redistributions of source code must retain the above copyright notice, this
-   list of conditions and the following disclaimer.
-
-2. Redistributions in binary form must reproduce the above copyright notice,
-   this list of conditions and the following disclaimer in the documentation
-   and/or other materials provided with the distribution.
-
-3. Neither the name of the copyright holder nor the names of its contributors
-   may be used to endorse or promote products derived from this software
-   without specific prior written permission.
-
-THIS SOFTWARE IS PROVIDED BY THE COPYRIGHT HOLDERS AND CONTRIBUTORS "AS IS" AND
-ANY EXPRESS OR IMPLIED WARRANTIES, INCLUDING, BUT NOT LIMITED TO, THE IMPLIED
-WARRANTIES OF MERCHANTABILITY AND FITNESS FOR A PARTICULAR PURPOSE ARE DISCLAIMED.
-IN NO EVENT SHALL THE COPYRIGHT HOLDER OR CONTRIBUTORS BE LIABLE FOR ANY DIRECT,
-INDIRECT, INCIDENTAL, SPECIAL, EXEMPLARY, OR CONSEQUENTIAL DAMAGES (INCLUDING,
-BUT NOT LIMITED TO, PROCUREMENT OF SUBSTITUTE GOODS OR SERVICES; LOSS OF USE, DATA,
-OR PROFITS; OR BUSINESS INTERRUPTION) HOWEVER CAUSED AND ON ANY THEORY OF LIABILITY,
-WHETHER IN CONTRACT, STRICT LIABILITY, OR TORT (INCLUDING NEGLIGENCE OR OTHERWISE)
-ARISING IN ANY WAY OUT OF THE USE OF THIS SOFTWARE, EVEN IF ADVISED OF THE
-POSSIBILITY OF SUCH DAMAGE.
-*/
-
+// Copyright 1998-2019 Epic Games, Inc. All Rights Reserved.
 
 #pragma once
 
diff --git a/Engine/Plugins/Runtime/Steam/SteamVR/Source/SteamVREditor/Public/SteamVREditorStyle.h b/Engine/Plugins/Runtime/Steam/SteamVR/Source/SteamVREditor/Public/SteamVREditorStyle.h
--- a/Engine/Plugins/Runtime/Steam/SteamVR/Source/SteamVREditor/Public/SteamVREditorStyle.h
+++ b/Engine/Plugins/Runtime/Steam/SteamVR/Source/SteamVREditor/Public/SteamVREditorStyle.h
@@ -1,32 +1,4 @@
-/*
-Copyright 2019 Valve Corporation under https://opensource.org/licenses/BSD-3-Clause
-
-Redistribution and use in source and binary forms, with or without modification,
-are permitted provided that the following conditions are met:
-
-1. Redistributions of source code must retain the above copyright notice, this
-   list of conditions and the following disclaimer.
-
-2. Redistributions in binary form must reproduce the above copyright notice,
-   this list of conditions and the following disclaimer in the documentation
-   and/or other materials provided with the distribution.
-
-3. Neither the name of the copyright holder nor the names of its contributors
-   may be used to endorse or promote products derived from this software
-   without specific prior written permission.
-
-THIS SOFTWARE IS PROVIDED BY THE COPYRIGHT HOLDERS AND CONTRIBUTORS "AS IS" AND
-ANY EXPRESS OR IMPLIED WARRANTIES, INCLUDING, BUT NOT LIMITED TO, THE IMPLIED
-WARRANTIES OF MERCHANTABILITY AND FITNESS FOR A PARTICULAR PURPOSE ARE DISCLAIMED.
-IN NO EVENT SHALL THE COPYRIGHT HOLDER OR CONTRIBUTORS BE LIABLE FOR ANY DIRECT,
-INDIRECT, INCIDENTAL, SPECIAL, EXEMPLARY, OR CONSEQUENTIAL DAMAGES (INCLUDING,
-BUT NOT LIMITED TO, PROCUREMENT OF SUBSTITUTE GOODS OR SERVICES; LOSS OF USE, DATA,
-OR PROFITS; OR BUSINESS INTERRUPTION) HOWEVER CAUSED AND ON ANY THEORY OF LIABILITY,
-WHETHER IN CONTRACT, STRICT LIABILITY, OR TORT (INCLUDING NEGLIGENCE OR OTHERWISE)
-ARISING IN ANY WAY OUT OF THE USE OF THIS SOFTWARE, EVEN IF ADVISED OF THE
-POSSIBILITY OF SUCH DAMAGE.
-*/
-
+// Copyright 1998-2019 Epic Games, Inc. All Rights Reserved.
 
 #pragma once
 
diff --git a/Engine/Plugins/Runtime/Steam/SteamVR/Source/SteamVREditor/SteamVREditor.Build.cs b/Engine/Plugins/Runtime/Steam/SteamVR/Source/SteamVREditor/SteamVREditor.Build.cs
--- a/Engine/Plugins/Runtime/Steam/SteamVR/Source/SteamVREditor/SteamVREditor.Build.cs
+++ b/Engine/Plugins/Runtime/Steam/SteamVR/Source/SteamVREditor/SteamVREditor.Build.cs
@@ -1,3 +1,5 @@
+// Copyright 1998-2019 Epic Games, Inc. All Rights Reserved.
+
 using UnrealBuildTool;
 using System.IO;
 
diff --git a/Engine/Plugins/Runtime/Steam/SteamVR/Source/SteamVRInput/SteamVRInput.Build.cs b/Engine/Plugins/Runtime/Steam/SteamVR/Source/SteamVRInput/SteamVRInput.Build.cs
--- a/Engine/Plugins/Runtime/Steam/SteamVR/Source/SteamVRInput/SteamVRInput.Build.cs
+++ b/Engine/Plugins/Runtime/Steam/SteamVR/Source/SteamVRInput/SteamVRInput.Build.cs
@@ -1,3 +1,5 @@
+// Copyright 1998-2019 Epic Games, Inc. All Rights Reserved.
+
 using UnrealBuildTool;
 using System.IO;
 
diff --git a/Engine/Plugins/Runtime/Steam/SteamVR/Source/SteamVRInputDevice/Private/SteamVRInputDevice.cpp b/Engine/Plugins/Runtime/Steam/SteamVR/Source/SteamVRInputDevice/Private/SteamVRInputDevice.cpp
--- a/Engine/Plugins/Runtime/Steam/SteamVR/Source/SteamVRInputDevice/Private/SteamVRInputDevice.cpp
+++ b/Engine/Plugins/Runtime/Steam/SteamVR/Source/SteamVRInputDevice/Private/SteamVRInputDevice.cpp
@@ -1119,10 +1119,7 @@ void FSteamVRInputDevice::InitControllerMappings()
 	{
 		for (unsigned int hand = 0; hand < k_unMaxTrackedDeviceCount; ++hand)
 		{
-			if (DeviceToControllerMap[id] < SteamVRInputDeviceConstants::MaxUnrealControllers && hand < k_unMaxTrackedDeviceCount)
-			{
-				UnrealControllerIdAndHandToDeviceIdMap[DeviceToControllerMap[id]][hand] = INDEX_NONE;
-			}
+			UnrealControllerIdAndHandToDeviceIdMap[id][hand] = INDEX_NONE;
 		}
 	}
 
diff --git a/Engine/Plugins/Runtime/Steam/SteamVR/Source/SteamVRInputDevice/Private/SteamVRTrackingRefComponent.cpp b/Engine/Plugins/Runtime/Steam/SteamVR/Source/SteamVRInputDevice/Private/SteamVRTrackingRefComponent.cpp
--- a/Engine/Plugins/Runtime/Steam/SteamVR/Source/SteamVRInputDevice/Private/SteamVRTrackingRefComponent.cpp
+++ b/Engine/Plugins/Runtime/Steam/SteamVR/Source/SteamVRInputDevice/Private/SteamVRTrackingRefComponent.cpp
@@ -44,7 +44,7 @@ bool USteamVRTrackingReferences::ShowTrackingReferences(UStaticMesh* TrackingRef
 				//UE_LOG(LogSteamVRTrackingRefComponent, Warning, TEXT("[TRACKING REFERENCE] Found the following tracking device: [%i] %s"), id, *stringCache);
 
 				TrackedDevicePose_t TrackedDevicePose;
-				VRCompositor()->GetLastPoseForTrackedDeviceIndex(id, &TrackedDevicePose, NULL);
+				VRCompositor()->GetLastPoseForTrackedDeviceIndex(id, &TrackedDevicePose, nullptr);
 
 				// Get SteamVR Transform Matrix for this tracking reference
 				HmdMatrix34_t Matrix = TrackedDevicePose.mDeviceToAbsoluteTracking;
@@ -173,9 +173,10 @@ void USteamVRTrackingReferences::TickComponent(float DeltaTime, ELevelTick TickT
 				FName DeviceClass = GetDeviceClass(id);
 
 				// Get device model info
-				char buf[k_unMaxPropertyStringSize];
-				uint32 StringBytes = VRSystem()->GetStringTrackedDeviceProperty(id, ETrackedDeviceProperty::Prop_ModelNumber_String, buf, sizeof(buf));
-				FString DeviceModel = *FString(UTF8_TO_TCHAR(buf));
+				TArray<char> Buffer;
+				Buffer.AddUninitialized(k_unMaxPropertyStringSize);
+				uint32 StringBytes = VRSystem()->GetStringTrackedDeviceProperty(id, ETrackedDeviceProperty::Prop_ModelNumber_String, Buffer.GetData(), Buffer.Num());
+				FString DeviceModel = *FString(UTF8_TO_TCHAR(Buffer.GetData()));
 
 				UE_LOG(LogSteamVRTrackingRefComponent, Warning, TEXT("Found device [%i] %s"), id, *DeviceModel);
 			}
diff --git a/Engine/Plugins/Runtime/Steam/SteamVR/Source/SteamVRInputDevice/SteamVRInputDevice.Build.cs b/Engine/Plugins/Runtime/Steam/SteamVR/Source/SteamVRInputDevice/SteamVRInputDevice.Build.cs
--- a/Engine/Plugins/Runtime/Steam/SteamVR/Source/SteamVRInputDevice/SteamVRInputDevice.Build.cs
+++ b/Engine/Plugins/Runtime/Steam/SteamVR/Source/SteamVRInputDevice/SteamVRInputDevice.Build.cs
@@ -1,3 +1,5 @@
+// Copyright 1998-2019 Epic Games, Inc. All Rights Reserved.
+
 using UnrealBuildTool;
 using System.IO;
 
