Move SteamVR Input to the Edit menu

diff --git a/Engine/Plugins/Runtime/Steam/SteamVR/Source/SteamVREditor/Private/SteamVREditor.cpp b/Engine/Plugins/Runtime/Steam/SteamVR/Source/SteamVREditor/Private/SteamVREditor.cpp
--- a/Engine/Plugins/Runtime/Steam/SteamVR/Source/SteamVREditor/Private/SteamVREditor.cpp
+++ b/Engine/Plugins/Runtime/Steam/SteamVR/Source/SteamVREditor/Private/SteamVREditor.cpp
@@ -60,20 +60,11 @@ void FSteamVREditorModule::StartupModule()
 		FCanExecuteAction());
 
 	FLevelEditorModule& LevelEditorModule = FModuleManager::LoadModuleChecked<FLevelEditorModule>("LevelEditor");
-	
-	{
-		TSharedPtr<FExtender> MenuExtender = MakeShareable(new FExtender());
-		MenuExtender->AddMenuExtension("WindowLayout", EExtensionHook::After, PluginCommands, FMenuExtensionDelegate::CreateRaw(this, &FSteamVREditorModule::AddMenuExtension));
 
-		LevelEditorModule.GetMenuExtensibilityManager()->AddExtender(MenuExtender);
-	}
-	
-	{
-		TSharedPtr<FExtender> ToolbarExtender = MakeShareable(new FExtender);
-		ToolbarExtender->AddToolBarExtension("Settings", EExtensionHook::After, PluginCommands, FToolBarExtensionDelegate::CreateRaw(this, &FSteamVREditorModule::AddToolbarExtension));
-		
-		LevelEditorModule.GetToolBarExtensibilityManager()->AddExtender(ToolbarExtender);
-	}
+	TSharedPtr<FExtender> MenuExtender = MakeShareable(new FExtender());
+	MenuExtender->AddMenuExtension("EditLocalTabSpawners", EExtensionHook::After, PluginCommands, FMenuExtensionDelegate::CreateRaw(this, &FSteamVREditorModule::AddMenuExtension));
+
+	LevelEditorModule.GetMenuExtensibilityManager()->AddExtender(MenuExtender);
 }
 
 void FSteamVREditorModule::ShutdownModule()
@@ -281,35 +272,26 @@ bool FSteamVREditorModule::AddUniqueActionMapping(TArray<FInputActionKeyMapping>
 }
 
 void FSteamVREditorModule::AddMenuExtension(FMenuBuilder& Builder)
-{
-	Builder.AddMenuEntry(FSteamVREditorCommands::Get().PluginAction);
-}
-
-void FSteamVREditorModule::AddToolbarExtension(FToolBarBuilder& Builder)
 {
 	FSteamVREditorStyle MenuStyle = FSteamVREditorStyle();
 	MenuStyle.Initialize();
 
-	Builder.AddComboButton(
-		FUIAction(FExecuteAction::CreateRaw(this, &FSteamVREditorModule::PluginButtonClicked)),
-		FOnGetContent::CreateRaw(this, &FSteamVREditorModule::FillComboButton, PluginCommands),
+	Builder.AddSubMenu(
 		LOCTEXT("SteamVRInputBtn", "SteamVR Input"),
 		LOCTEXT("SteamVRInputBtnTootlip", "SteamVR Input"),
-		FSlateIcon(FSteamVREditorStyle::GetStyleSetName(), "SteamVREditor.PluginAction")
+		FNewMenuDelegate::CreateRaw(this, &FSteamVREditorModule::MakeMenu),
+		false, FSlateIcon(FSteamVREditorStyle::GetStyleSetName(), "SteamVREditor.PluginAction")
 	);
 }
 
-TSharedRef<SWidget> FSteamVREditorModule::FillComboButton(TSharedPtr<class FUICommandList> Commands)
+void FSteamVREditorModule::MakeMenu(FMenuBuilder& MenuBuilder)
 {
-	FMenuBuilder MenuBuilder(true, Commands);
-
+	//MenuBuilder.AddMenuEntry(FSteamVREditorCommands::Get().PluginAction);
 	MenuBuilder.AddMenuEntry(FSteamVREditorCommands::Get().JsonActionManifest, NAME_None, TAttribute<FText>(), TAttribute<FText>(), FSlateIcon(FSteamVREditorStyle::GetStyleSetName(), "SteamVREditor.JsonActionManifest"));
 	MenuBuilder.AddMenuEntry(FSteamVREditorCommands::Get().JsonControllerBindings, NAME_None, TAttribute<FText>(), TAttribute<FText>(), FSlateIcon(FSteamVREditorStyle::GetStyleSetName(), "SteamVREditor.JsonControllerBindings"));
 	MenuBuilder.AddMenuEntry(FSteamVREditorCommands::Get().ReloadActionManifest, NAME_None, TAttribute<FText>(), TAttribute<FText>(), FSlateIcon(FSteamVREditorStyle::GetStyleSetName(), "SteamVREditor.ReloadActionManifest"));
 	MenuBuilder.AddMenuEntry(FSteamVREditorCommands::Get().LaunchBindingsURL, NAME_None, TAttribute<FText>(), TAttribute<FText>(), FSlateIcon(FSteamVREditorStyle::GetStyleSetName(), "SteamVREditor.LaunchBindingsURL"));
 	//MenuBuilder.AddMenuEntry(FSteamVREditorCommands::Get().AddSampleInputs, NAME_None, TAttribute<FText>(), TAttribute<FText>(), FSlateIcon(FSteamVREditorStyle::GetStyleSetName(), "SteamVREditor.AddSampleInputs"));
-
-	return MenuBuilder.MakeWidget();
 }
 
 #undef LOCTEXT_NAMESPACE
diff --git a/Engine/Plugins/Runtime/Steam/SteamVR/Source/SteamVREditor/Private/SteamVREditorStyle.cpp b/Engine/Plugins/Runtime/Steam/SteamVR/Source/SteamVREditor/Private/SteamVREditorStyle.cpp
--- a/Engine/Plugins/Runtime/Steam/SteamVR/Source/SteamVREditor/Private/SteamVREditorStyle.cpp
+++ b/Engine/Plugins/Runtime/Steam/SteamVR/Source/SteamVREditor/Private/SteamVREditorStyle.cpp
@@ -43,7 +43,7 @@ TSharedRef< FSlateStyleSet > FSteamVREditorStyle::Create()
 	TSharedRef< FSlateStyleSet > Style = MakeShareable(new FSlateStyleSet("SteamVREditorStyle"));
 	Style->SetContentRoot(IPluginManager::Get().FindPlugin("SteamVR")->GetBaseDir() / TEXT("Resources"));
 
-	Style->Set("SteamVREditor.PluginAction", new IMAGE_BRUSH(TEXT("Icon_PluginAction"), Icon40x40));
+	Style->Set("SteamVREditor.PluginAction", new IMAGE_BRUSH(TEXT("Icon_PluginAction_16x"), Icon16x16));
 	Style->Set("SteamVREditor.JsonActionManifest", new IMAGE_BRUSH(TEXT("Icon_regenerate_action_manifest"), Icon16x16));
 	Style->Set("SteamVREditor.JsonControllerBindings", new IMAGE_BRUSH(TEXT("Icon_regenerate_controller_bindings"), Icon16x16));
 	Style->Set("SteamVREditor.ReloadActionManifest", new IMAGE_BRUSH(TEXT("Icon_reload_action_manifest"), Icon16x16));
diff --git a/Engine/Plugins/Runtime/Steam/SteamVR/Source/SteamVREditor/Public/SteamVREditor.h b/Engine/Plugins/Runtime/Steam/SteamVR/Source/SteamVREditor/Public/SteamVREditor.h
--- a/Engine/Plugins/Runtime/Steam/SteamVR/Source/SteamVREditor/Public/SteamVREditor.h
+++ b/Engine/Plugins/Runtime/Steam/SteamVR/Source/SteamVREditor/Public/SteamVREditor.h
@@ -29,12 +29,11 @@ public:
 
 private:
 
-	void AddToolbarExtension(FToolBarBuilder& Builder);
 	void AddMenuExtension(FMenuBuilder& Builder);
 	bool AddUniqueAxisMapping(TArray<FInputAxisKeyMapping> ExistingAxisKeys, UInputSettings* InputSettings, FName ActionName, FKey ActionKey);
 	bool AddUniqueActionMapping(TArray<FInputActionKeyMapping> ExistingActionKeys, UInputSettings* InputSettings, FName ActionName, FKey ActionKey);
 
 	TSharedPtr<class FUICommandList> PluginCommands;
-	TSharedRef<SWidget> FillComboButton(TSharedPtr<class FUICommandList> Commands);
+	void MakeMenu(FMenuBuilder& MenuBuilder);
 
 };
