The assert on IsInGameThread is invalid in the depth check case, since this will always be called from the rendering thread as part of the creation/update of common depth targets.

diff --git a/Engine/Plugins/Runtime/Steam/SteamVR/Source/SteamVR/Private/SteamVRHMD.cpp b/Engine/Plugins/Runtime/Steam/SteamVR/Source/SteamVR/Private/SteamVRHMD.cpp
--- a/Engine/Plugins/Runtime/Steam/SteamVR/Source/SteamVR/Private/SteamVRHMD.cpp
+++ b/Engine/Plugins/Runtime/Steam/SteamVR/Source/SteamVR/Private/SteamVRHMD.cpp
@@ -1573,7 +1573,7 @@ bool FSteamVRHMD::NeedReAllocateViewportRenderTarget(const FViewport& Viewport)
 
 bool FSteamVRHMD::NeedReAllocateDepthTexture(const TRefCountPtr<struct IPooledRenderTarget> & DepthTarget)
 {
-	check(IsInGameThread());
+	check(IsInRenderingThread());
 
 	// Check the dimensions of the currently stored depth swapchain vs the current rendering swapchain.
 	if (pBridge->GetSwapChain()->GetTexture2D()->GetSizeX() != pBridge->GetDepthSwapChain()->GetTexture2D()->GetSizeX() ||
