Provide helpful prompts to the SteamVR Input users if they're overwriting bindings
Update copyright notices to reflect new agreement

diff --git a/Engine/Plugins/Runtime/Steam/SteamVR/Source/SteamVREditor/Private/AnimGraphNode_SteamVRInputAnimPose.cpp b/Engine/Plugins/Runtime/Steam/SteamVR/Source/SteamVREditor/Private/AnimGraphNode_SteamVRInputAnimPose.cpp
--- a/Engine/Plugins/Runtime/Steam/SteamVR/Source/SteamVREditor/Private/AnimGraphNode_SteamVRInputAnimPose.cpp
+++ b/Engine/Plugins/Runtime/Steam/SteamVR/Source/SteamVREditor/Private/AnimGraphNode_SteamVRInputAnimPose.cpp
@@ -1,4 +1,32 @@
-// Copyright 1998-2019 Epic Games, Inc. All Rights Reserved.
+/*
+Copyright 2019 Valve Corporation under https://opensource.org/licenses/BSD-3-Clause
+This code includes modifications by Epic Games.  Modifications (c) 2019 Epic Games, Inc.
+
+Redistribution and use in source and binary forms, with or without modification,
+are permitted provided that the following conditions are met:
+
+1. Redistributions of source code must retain the above copyright notice, this
+   list of conditions and the following disclaimer.
+
+2. Redistributions in binary form must reproduce the above copyright notice,
+   this list of conditions and the following disclaimer in the documentation
+   and/or other materials provided with the distribution.
+
+3. Neither the name of the copyright holder nor the names of its contributors
+   may be used to endorse or promote products derived from this software
+   without specific prior written permission.
+
+THIS SOFTWARE IS PROVIDED BY THE COPYRIGHT HOLDERS AND CONTRIBUTORS "AS IS" AND
+ANY EXPRESS OR IMPLIED WARRANTIES, INCLUDING, BUT NOT LIMITED TO, THE IMPLIED
+WARRANTIES OF MERCHANTABILITY AND FITNESS FOR A PARTICULAR PURPOSE ARE DISCLAIMED.
+IN NO EVENT SHALL THE COPYRIGHT HOLDER OR CONTRIBUTORS BE LIABLE FOR ANY DIRECT,
+INDIRECT, INCIDENTAL, SPECIAL, EXEMPLARY, OR CONSEQUENTIAL DAMAGES (INCLUDING,
+BUT NOT LIMITED TO, PROCUREMENT OF SUBSTITUTE GOODS OR SERVICES; LOSS OF USE, DATA,
+OR PROFITS; OR BUSINESS INTERRUPTION) HOWEVER CAUSED AND ON ANY THEORY OF LIABILITY,
+WHETHER IN CONTRACT, STRICT LIABILITY, OR TORT (INCLUDING NEGLIGENCE OR OTHERWISE)
+ARISING IN ANY WAY OUT OF THE USE OF THIS SOFTWARE, EVEN IF ADVISED OF THE
+POSSIBILITY OF SUCH DAMAGE.
+*/
 
 #include "AnimGraphNode_SteamVRInputAnimPose.h"
 #include "SteamVREditor.h"
diff --git a/Engine/Plugins/Runtime/Steam/SteamVR/Source/SteamVREditor/Private/AnimGraphNode_SteamVRSetWristTransform.cpp b/Engine/Plugins/Runtime/Steam/SteamVR/Source/SteamVREditor/Private/AnimGraphNode_SteamVRSetWristTransform.cpp
--- a/Engine/Plugins/Runtime/Steam/SteamVR/Source/SteamVREditor/Private/AnimGraphNode_SteamVRSetWristTransform.cpp
+++ b/Engine/Plugins/Runtime/Steam/SteamVR/Source/SteamVREditor/Private/AnimGraphNode_SteamVRSetWristTransform.cpp
@@ -1,4 +1,32 @@
-// Copyright 1998-2019 Epic Games, Inc. All Rights Reserved.
+/*
+Copyright 2019 Valve Corporation under https://opensource.org/licenses/BSD-3-Clause
+This code includes modifications by Epic Games.  Modifications (c) 2019 Epic Games, Inc.
+
+Redistribution and use in source and binary forms, with or without modification,
+are permitted provided that the following conditions are met:
+
+1. Redistributions of source code must retain the above copyright notice, this
+   list of conditions and the following disclaimer.
+
+2. Redistributions in binary form must reproduce the above copyright notice,
+   this list of conditions and the following disclaimer in the documentation
+   and/or other materials provided with the distribution.
+
+3. Neither the name of the copyright holder nor the names of its contributors
+   may be used to endorse or promote products derived from this software
+   without specific prior written permission.
+
+THIS SOFTWARE IS PROVIDED BY THE COPYRIGHT HOLDERS AND CONTRIBUTORS "AS IS" AND
+ANY EXPRESS OR IMPLIED WARRANTIES, INCLUDING, BUT NOT LIMITED TO, THE IMPLIED
+WARRANTIES OF MERCHANTABILITY AND FITNESS FOR A PARTICULAR PURPOSE ARE DISCLAIMED.
+IN NO EVENT SHALL THE COPYRIGHT HOLDER OR CONTRIBUTORS BE LIABLE FOR ANY DIRECT,
+INDIRECT, INCIDENTAL, SPECIAL, EXEMPLARY, OR CONSEQUENTIAL DAMAGES (INCLUDING,
+BUT NOT LIMITED TO, PROCUREMENT OF SUBSTITUTE GOODS OR SERVICES; LOSS OF USE, DATA,
+OR PROFITS; OR BUSINESS INTERRUPTION) HOWEVER CAUSED AND ON ANY THEORY OF LIABILITY,
+WHETHER IN CONTRACT, STRICT LIABILITY, OR TORT (INCLUDING NEGLIGENCE OR OTHERWISE)
+ARISING IN ANY WAY OUT OF THE USE OF THIS SOFTWARE, EVEN IF ADVISED OF THE
+POSSIBILITY OF SUCH DAMAGE.
+*/
 
 #include "AnimGraphNode_SteamVRSetWristTransform.h"
 #include "SteamVREditor.h"
diff --git a/Engine/Plugins/Runtime/Steam/SteamVR/Source/SteamVREditor/Private/SteamVREditor.cpp b/Engine/Plugins/Runtime/Steam/SteamVR/Source/SteamVREditor/Private/SteamVREditor.cpp
--- a/Engine/Plugins/Runtime/Steam/SteamVR/Source/SteamVREditor/Private/SteamVREditor.cpp
+++ b/Engine/Plugins/Runtime/Steam/SteamVR/Source/SteamVREditor/Private/SteamVREditor.cpp
@@ -1,4 +1,32 @@
-// Copyright 1998-2019 Epic Games, Inc. All Rights Reserved.
+/*
+Copyright 2019 Valve Corporation under https://opensource.org/licenses/BSD-3-Clause
+This code includes modifications by Epic Games.  Modifications (c) 2019 Epic Games, Inc.
+
+Redistribution and use in source and binary forms, with or without modification,
+are permitted provided that the following conditions are met:
+
+1. Redistributions of source code must retain the above copyright notice, this
+   list of conditions and the following disclaimer.
+
+2. Redistributions in binary form must reproduce the above copyright notice,
+   this list of conditions and the following disclaimer in the documentation
+   and/or other materials provided with the distribution.
+
+3. Neither the name of the copyright holder nor the names of its contributors
+   may be used to endorse or promote products derived from this software
+   without specific prior written permission.
+
+THIS SOFTWARE IS PROVIDED BY THE COPYRIGHT HOLDERS AND CONTRIBUTORS "AS IS" AND
+ANY EXPRESS OR IMPLIED WARRANTIES, INCLUDING, BUT NOT LIMITED TO, THE IMPLIED
+WARRANTIES OF MERCHANTABILITY AND FITNESS FOR A PARTICULAR PURPOSE ARE DISCLAIMED.
+IN NO EVENT SHALL THE COPYRIGHT HOLDER OR CONTRIBUTORS BE LIABLE FOR ANY DIRECT,
+INDIRECT, INCIDENTAL, SPECIAL, EXEMPLARY, OR CONSEQUENTIAL DAMAGES (INCLUDING,
+BUT NOT LIMITED TO, PROCUREMENT OF SUBSTITUTE GOODS OR SERVICES; LOSS OF USE, DATA,
+OR PROFITS; OR BUSINESS INTERRUPTION) HOWEVER CAUSED AND ON ANY THEORY OF LIABILITY,
+WHETHER IN CONTRACT, STRICT LIABILITY, OR TORT (INCLUDING NEGLIGENCE OR OTHERWISE)
+ARISING IN ANY WAY OUT OF THE USE OF THIS SOFTWARE, EVEN IF ADVISED OF THE
+POSSIBILITY OF SUCH DAMAGE.
+*/
 
 #pragma once
 
@@ -8,8 +36,9 @@
 #include "SteamVREditorStyle.h"
 #include "SteamVREditorCommands.h"
 #include "SteamVRControllerKeys.h"
-#include "SteamVRInputSettings.h"
 #include "LevelEditor.h"
+#include "Engine/Engine.h"
+#include "IXRTrackingSystem.h"
 
 static const FName SteamVREditorTabName("SteamVREditor");
 
@@ -60,13 +89,14 @@ void FSteamVREditorModule::StartupModule()
 		FExecuteAction::CreateRaw(this, &FSteamVREditorModule::AddSampleInputs),
 		FCanExecuteAction());
 
-	FLevelEditorModule& LevelEditorModule = FModuleManager::LoadModuleChecked<FLevelEditorModule>("LevelEditor");
-	const USteamVRInputSettings* Settings = GetDefault<USteamVRInputSettings>();
-	if (Settings->bDeveloperMode)
+	// Only add the toolbar if SteamVR is the currently active tracking system
+	static FName SystemName(TEXT("SteamVR"));
+	if (GEngine && GEngine->XRSystem.IsValid() && (GEngine->XRSystem->GetSystemName() == SystemName))
 	{
 		TSharedPtr<FExtender> ToolbarExtender = MakeShareable(new FExtender);
 		ToolbarExtender->AddToolBarExtension("Settings", EExtensionHook::After, PluginCommands, FToolBarExtensionDelegate::CreateRaw(this, &FSteamVREditorModule::AddToolbarExtension));
-		
+
+		FLevelEditorModule& LevelEditorModule = FModuleManager::LoadModuleChecked<FLevelEditorModule>("LevelEditor");
 		LevelEditorModule.GetToolBarExtensibilityManager()->AddExtender(ToolbarExtender);
 	}
 }
diff --git a/Engine/Plugins/Runtime/Steam/SteamVR/Source/SteamVREditor/Private/SteamVREditorStyle.cpp b/Engine/Plugins/Runtime/Steam/SteamVR/Source/SteamVREditor/Private/SteamVREditorStyle.cpp
--- a/Engine/Plugins/Runtime/Steam/SteamVR/Source/SteamVREditor/Private/SteamVREditorStyle.cpp
+++ b/Engine/Plugins/Runtime/Steam/SteamVR/Source/SteamVREditor/Private/SteamVREditorStyle.cpp
@@ -1,4 +1,32 @@
-// Copyright 1998-2019 Epic Games, Inc. All Rights Reserved.
+/*
+Copyright 2019 Valve Corporation under https://opensource.org/licenses/BSD-3-Clause
+This code includes modifications by Epic Games.  Modifications (c) 2019 Epic Games, Inc.
+
+Redistribution and use in source and binary forms, with or without modification,
+are permitted provided that the following conditions are met:
+
+1. Redistributions of source code must retain the above copyright notice, this
+   list of conditions and the following disclaimer.
+
+2. Redistributions in binary form must reproduce the above copyright notice,
+   this list of conditions and the following disclaimer in the documentation
+   and/or other materials provided with the distribution.
+
+3. Neither the name of the copyright holder nor the names of its contributors
+   may be used to endorse or promote products derived from this software
+   without specific prior written permission.
+
+THIS SOFTWARE IS PROVIDED BY THE COPYRIGHT HOLDERS AND CONTRIBUTORS "AS IS" AND
+ANY EXPRESS OR IMPLIED WARRANTIES, INCLUDING, BUT NOT LIMITED TO, THE IMPLIED
+WARRANTIES OF MERCHANTABILITY AND FITNESS FOR A PARTICULAR PURPOSE ARE DISCLAIMED.
+IN NO EVENT SHALL THE COPYRIGHT HOLDER OR CONTRIBUTORS BE LIABLE FOR ANY DIRECT,
+INDIRECT, INCIDENTAL, SPECIAL, EXEMPLARY, OR CONSEQUENTIAL DAMAGES (INCLUDING,
+BUT NOT LIMITED TO, PROCUREMENT OF SUBSTITUTE GOODS OR SERVICES; LOSS OF USE, DATA,
+OR PROFITS; OR BUSINESS INTERRUPTION) HOWEVER CAUSED AND ON ANY THEORY OF LIABILITY,
+WHETHER IN CONTRACT, STRICT LIABILITY, OR TORT (INCLUDING NEGLIGENCE OR OTHERWISE)
+ARISING IN ANY WAY OUT OF THE USE OF THIS SOFTWARE, EVEN IF ADVISED OF THE
+POSSIBILITY OF SUCH DAMAGE.
+*/
 
 #include "SteamVREditorStyle.h"
 #include "SteamVREditor.h"
diff --git a/Engine/Plugins/Runtime/Steam/SteamVR/Source/SteamVREditor/Public/AnimGraphNode_SteamVRInputAnimPose.h b/Engine/Plugins/Runtime/Steam/SteamVR/Source/SteamVREditor/Public/AnimGraphNode_SteamVRInputAnimPose.h
--- a/Engine/Plugins/Runtime/Steam/SteamVR/Source/SteamVREditor/Public/AnimGraphNode_SteamVRInputAnimPose.h
+++ b/Engine/Plugins/Runtime/Steam/SteamVR/Source/SteamVREditor/Public/AnimGraphNode_SteamVRInputAnimPose.h
@@ -1,4 +1,32 @@
-// Copyright 1998-2019 Epic Games, Inc. All Rights Reserved.
+/*
+Copyright 2019 Valve Corporation under https://opensource.org/licenses/BSD-3-Clause
+This code includes modifications by Epic Games.  Modifications (c) 2019 Epic Games, Inc.
+
+Redistribution and use in source and binary forms, with or without modification,
+are permitted provided that the following conditions are met:
+
+1. Redistributions of source code must retain the above copyright notice, this
+   list of conditions and the following disclaimer.
+
+2. Redistributions in binary form must reproduce the above copyright notice,
+   this list of conditions and the following disclaimer in the documentation
+   and/or other materials provided with the distribution.
+
+3. Neither the name of the copyright holder nor the names of its contributors
+   may be used to endorse or promote products derived from this software
+   without specific prior written permission.
+
+THIS SOFTWARE IS PROVIDED BY THE COPYRIGHT HOLDERS AND CONTRIBUTORS "AS IS" AND
+ANY EXPRESS OR IMPLIED WARRANTIES, INCLUDING, BUT NOT LIMITED TO, THE IMPLIED
+WARRANTIES OF MERCHANTABILITY AND FITNESS FOR A PARTICULAR PURPOSE ARE DISCLAIMED.
+IN NO EVENT SHALL THE COPYRIGHT HOLDER OR CONTRIBUTORS BE LIABLE FOR ANY DIRECT,
+INDIRECT, INCIDENTAL, SPECIAL, EXEMPLARY, OR CONSEQUENTIAL DAMAGES (INCLUDING,
+BUT NOT LIMITED TO, PROCUREMENT OF SUBSTITUTE GOODS OR SERVICES; LOSS OF USE, DATA,
+OR PROFITS; OR BUSINESS INTERRUPTION) HOWEVER CAUSED AND ON ANY THEORY OF LIABILITY,
+WHETHER IN CONTRACT, STRICT LIABILITY, OR TORT (INCLUDING NEGLIGENCE OR OTHERWISE)
+ARISING IN ANY WAY OUT OF THE USE OF THIS SOFTWARE, EVEN IF ADVISED OF THE
+POSSIBILITY OF SUCH DAMAGE.
+*/
 
 #pragma once
 
diff --git a/Engine/Plugins/Runtime/Steam/SteamVR/Source/SteamVREditor/Public/AnimGraphNode_SteamVRSetWristTransform.h b/Engine/Plugins/Runtime/Steam/SteamVR/Source/SteamVREditor/Public/AnimGraphNode_SteamVRSetWristTransform.h
--- a/Engine/Plugins/Runtime/Steam/SteamVR/Source/SteamVREditor/Public/AnimGraphNode_SteamVRSetWristTransform.h
+++ b/Engine/Plugins/Runtime/Steam/SteamVR/Source/SteamVREditor/Public/AnimGraphNode_SteamVRSetWristTransform.h
@@ -1,4 +1,32 @@
-// Copyright 1998-2019 Epic Games, Inc. All Rights Reserved.
+/*
+Copyright 2019 Valve Corporation under https://opensource.org/licenses/BSD-3-Clause
+This code includes modifications by Epic Games.  Modifications (c) 2019 Epic Games, Inc.
+
+Redistribution and use in source and binary forms, with or without modification,
+are permitted provided that the following conditions are met:
+
+1. Redistributions of source code must retain the above copyright notice, this
+   list of conditions and the following disclaimer.
+
+2. Redistributions in binary form must reproduce the above copyright notice,
+   this list of conditions and the following disclaimer in the documentation
+   and/or other materials provided with the distribution.
+
+3. Neither the name of the copyright holder nor the names of its contributors
+   may be used to endorse or promote products derived from this software
+   without specific prior written permission.
+
+THIS SOFTWARE IS PROVIDED BY THE COPYRIGHT HOLDERS AND CONTRIBUTORS "AS IS" AND
+ANY EXPRESS OR IMPLIED WARRANTIES, INCLUDING, BUT NOT LIMITED TO, THE IMPLIED
+WARRANTIES OF MERCHANTABILITY AND FITNESS FOR A PARTICULAR PURPOSE ARE DISCLAIMED.
+IN NO EVENT SHALL THE COPYRIGHT HOLDER OR CONTRIBUTORS BE LIABLE FOR ANY DIRECT,
+INDIRECT, INCIDENTAL, SPECIAL, EXEMPLARY, OR CONSEQUENTIAL DAMAGES (INCLUDING,
+BUT NOT LIMITED TO, PROCUREMENT OF SUBSTITUTE GOODS OR SERVICES; LOSS OF USE, DATA,
+OR PROFITS; OR BUSINESS INTERRUPTION) HOWEVER CAUSED AND ON ANY THEORY OF LIABILITY,
+WHETHER IN CONTRACT, STRICT LIABILITY, OR TORT (INCLUDING NEGLIGENCE OR OTHERWISE)
+ARISING IN ANY WAY OUT OF THE USE OF THIS SOFTWARE, EVEN IF ADVISED OF THE
+POSSIBILITY OF SUCH DAMAGE.
+*/
 
 #pragma once
 
diff --git a/Engine/Plugins/Runtime/Steam/SteamVR/Source/SteamVREditor/Public/SteamVREditor.h b/Engine/Plugins/Runtime/Steam/SteamVR/Source/SteamVREditor/Public/SteamVREditor.h
--- a/Engine/Plugins/Runtime/Steam/SteamVR/Source/SteamVREditor/Public/SteamVREditor.h
+++ b/Engine/Plugins/Runtime/Steam/SteamVR/Source/SteamVREditor/Public/SteamVREditor.h
@@ -1,4 +1,32 @@
-// Copyright 1998-2019 Epic Games, Inc. All Rights Reserved.
+/*
+Copyright 2019 Valve Corporation under https://opensource.org/licenses/BSD-3-Clause
+This code includes modifications by Epic Games.  Modifications (c) 2019 Epic Games, Inc.
+
+Redistribution and use in source and binary forms, with or without modification,
+are permitted provided that the following conditions are met:
+
+1. Redistributions of source code must retain the above copyright notice, this
+   list of conditions and the following disclaimer.
+
+2. Redistributions in binary form must reproduce the above copyright notice,
+   this list of conditions and the following disclaimer in the documentation
+   and/or other materials provided with the distribution.
+
+3. Neither the name of the copyright holder nor the names of its contributors
+   may be used to endorse or promote products derived from this software
+   without specific prior written permission.
+
+THIS SOFTWARE IS PROVIDED BY THE COPYRIGHT HOLDERS AND CONTRIBUTORS "AS IS" AND
+ANY EXPRESS OR IMPLIED WARRANTIES, INCLUDING, BUT NOT LIMITED TO, THE IMPLIED
+WARRANTIES OF MERCHANTABILITY AND FITNESS FOR A PARTICULAR PURPOSE ARE DISCLAIMED.
+IN NO EVENT SHALL THE COPYRIGHT HOLDER OR CONTRIBUTORS BE LIABLE FOR ANY DIRECT,
+INDIRECT, INCIDENTAL, SPECIAL, EXEMPLARY, OR CONSEQUENTIAL DAMAGES (INCLUDING,
+BUT NOT LIMITED TO, PROCUREMENT OF SUBSTITUTE GOODS OR SERVICES; LOSS OF USE, DATA,
+OR PROFITS; OR BUSINESS INTERRUPTION) HOWEVER CAUSED AND ON ANY THEORY OF LIABILITY,
+WHETHER IN CONTRACT, STRICT LIABILITY, OR TORT (INCLUDING NEGLIGENCE OR OTHERWISE)
+ARISING IN ANY WAY OUT OF THE USE OF THIS SOFTWARE, EVEN IF ADVISED OF THE
+POSSIBILITY OF SUCH DAMAGE.
+*/
 
 #pragma once
 
diff --git a/Engine/Plugins/Runtime/Steam/SteamVR/Source/SteamVREditor/Public/SteamVREditorCommands.cpp b/Engine/Plugins/Runtime/Steam/SteamVR/Source/SteamVREditor/Public/SteamVREditorCommands.cpp
--- a/Engine/Plugins/Runtime/Steam/SteamVR/Source/SteamVREditor/Public/SteamVREditorCommands.cpp
+++ b/Engine/Plugins/Runtime/Steam/SteamVR/Source/SteamVREditor/Public/SteamVREditorCommands.cpp
@@ -1,4 +1,32 @@
-// Copyright 1998-2019 Epic Games, Inc. All Rights Reserved.
+/*
+Copyright 2019 Valve Corporation under https://opensource.org/licenses/BSD-3-Clause
+This code includes modifications by Epic Games.  Modifications (c) 2019 Epic Games, Inc.
+
+Redistribution and use in source and binary forms, with or without modification,
+are permitted provided that the following conditions are met:
+
+1. Redistributions of source code must retain the above copyright notice, this
+   list of conditions and the following disclaimer.
+
+2. Redistributions in binary form must reproduce the above copyright notice,
+   this list of conditions and the following disclaimer in the documentation
+   and/or other materials provided with the distribution.
+
+3. Neither the name of the copyright holder nor the names of its contributors
+   may be used to endorse or promote products derived from this software
+   without specific prior written permission.
+
+THIS SOFTWARE IS PROVIDED BY THE COPYRIGHT HOLDERS AND CONTRIBUTORS "AS IS" AND
+ANY EXPRESS OR IMPLIED WARRANTIES, INCLUDING, BUT NOT LIMITED TO, THE IMPLIED
+WARRANTIES OF MERCHANTABILITY AND FITNESS FOR A PARTICULAR PURPOSE ARE DISCLAIMED.
+IN NO EVENT SHALL THE COPYRIGHT HOLDER OR CONTRIBUTORS BE LIABLE FOR ANY DIRECT,
+INDIRECT, INCIDENTAL, SPECIAL, EXEMPLARY, OR CONSEQUENTIAL DAMAGES (INCLUDING,
+BUT NOT LIMITED TO, PROCUREMENT OF SUBSTITUTE GOODS OR SERVICES; LOSS OF USE, DATA,
+OR PROFITS; OR BUSINESS INTERRUPTION) HOWEVER CAUSED AND ON ANY THEORY OF LIABILITY,
+WHETHER IN CONTRACT, STRICT LIABILITY, OR TORT (INCLUDING NEGLIGENCE OR OTHERWISE)
+ARISING IN ANY WAY OUT OF THE USE OF THIS SOFTWARE, EVEN IF ADVISED OF THE
+POSSIBILITY OF SUCH DAMAGE.
+*/
 
 #include "SteamVREditorCommands.h"
 #include "SteamVREditor.h"
diff --git a/Engine/Plugins/Runtime/Steam/SteamVR/Source/SteamVREditor/Public/SteamVREditorCommands.h b/Engine/Plugins/Runtime/Steam/SteamVR/Source/SteamVREditor/Public/SteamVREditorCommands.h
--- a/Engine/Plugins/Runtime/Steam/SteamVR/Source/SteamVREditor/Public/SteamVREditorCommands.h
+++ b/Engine/Plugins/Runtime/Steam/SteamVR/Source/SteamVREditor/Public/SteamVREditorCommands.h
@@ -1,4 +1,32 @@
-// Copyright 1998-2019 Epic Games, Inc. All Rights Reserved.
+/*
+Copyright 2019 Valve Corporation under https://opensource.org/licenses/BSD-3-Clause
+This code includes modifications by Epic Games.  Modifications (c) 2019 Epic Games, Inc.
+
+Redistribution and use in source and binary forms, with or without modification,
+are permitted provided that the following conditions are met:
+
+1. Redistributions of source code must retain the above copyright notice, this
+   list of conditions and the following disclaimer.
+
+2. Redistributions in binary form must reproduce the above copyright notice,
+   this list of conditions and the following disclaimer in the documentation
+   and/or other materials provided with the distribution.
+
+3. Neither the name of the copyright holder nor the names of its contributors
+   may be used to endorse or promote products derived from this software
+   without specific prior written permission.
+
+THIS SOFTWARE IS PROVIDED BY THE COPYRIGHT HOLDERS AND CONTRIBUTORS "AS IS" AND
+ANY EXPRESS OR IMPLIED WARRANTIES, INCLUDING, BUT NOT LIMITED TO, THE IMPLIED
+WARRANTIES OF MERCHANTABILITY AND FITNESS FOR A PARTICULAR PURPOSE ARE DISCLAIMED.
+IN NO EVENT SHALL THE COPYRIGHT HOLDER OR CONTRIBUTORS BE LIABLE FOR ANY DIRECT,
+INDIRECT, INCIDENTAL, SPECIAL, EXEMPLARY, OR CONSEQUENTIAL DAMAGES (INCLUDING,
+BUT NOT LIMITED TO, PROCUREMENT OF SUBSTITUTE GOODS OR SERVICES; LOSS OF USE, DATA,
+OR PROFITS; OR BUSINESS INTERRUPTION) HOWEVER CAUSED AND ON ANY THEORY OF LIABILITY,
+WHETHER IN CONTRACT, STRICT LIABILITY, OR TORT (INCLUDING NEGLIGENCE OR OTHERWISE)
+ARISING IN ANY WAY OUT OF THE USE OF THIS SOFTWARE, EVEN IF ADVISED OF THE
+POSSIBILITY OF SUCH DAMAGE.
+*/
 
 #pragma once
 
diff --git a/Engine/Plugins/Runtime/Steam/SteamVR/Source/SteamVREditor/Public/SteamVREditorStyle.h b/Engine/Plugins/Runtime/Steam/SteamVR/Source/SteamVREditor/Public/SteamVREditorStyle.h
--- a/Engine/Plugins/Runtime/Steam/SteamVR/Source/SteamVREditor/Public/SteamVREditorStyle.h
+++ b/Engine/Plugins/Runtime/Steam/SteamVR/Source/SteamVREditor/Public/SteamVREditorStyle.h
@@ -1,4 +1,32 @@
-// Copyright 1998-2019 Epic Games, Inc. All Rights Reserved.
+/*
+Copyright 2019 Valve Corporation under https://opensource.org/licenses/BSD-3-Clause
+This code includes modifications by Epic Games.  Modifications (c) 2019 Epic Games, Inc.
+
+Redistribution and use in source and binary forms, with or without modification,
+are permitted provided that the following conditions are met:
+
+1. Redistributions of source code must retain the above copyright notice, this
+   list of conditions and the following disclaimer.
+
+2. Redistributions in binary form must reproduce the above copyright notice,
+   this list of conditions and the following disclaimer in the documentation
+   and/or other materials provided with the distribution.
+
+3. Neither the name of the copyright holder nor the names of its contributors
+   may be used to endorse or promote products derived from this software
+   without specific prior written permission.
+
+THIS SOFTWARE IS PROVIDED BY THE COPYRIGHT HOLDERS AND CONTRIBUTORS "AS IS" AND
+ANY EXPRESS OR IMPLIED WARRANTIES, INCLUDING, BUT NOT LIMITED TO, THE IMPLIED
+WARRANTIES OF MERCHANTABILITY AND FITNESS FOR A PARTICULAR PURPOSE ARE DISCLAIMED.
+IN NO EVENT SHALL THE COPYRIGHT HOLDER OR CONTRIBUTORS BE LIABLE FOR ANY DIRECT,
+INDIRECT, INCIDENTAL, SPECIAL, EXEMPLARY, OR CONSEQUENTIAL DAMAGES (INCLUDING,
+BUT NOT LIMITED TO, PROCUREMENT OF SUBSTITUTE GOODS OR SERVICES; LOSS OF USE, DATA,
+OR PROFITS; OR BUSINESS INTERRUPTION) HOWEVER CAUSED AND ON ANY THEORY OF LIABILITY,
+WHETHER IN CONTRACT, STRICT LIABILITY, OR TORT (INCLUDING NEGLIGENCE OR OTHERWISE)
+ARISING IN ANY WAY OUT OF THE USE OF THIS SOFTWARE, EVEN IF ADVISED OF THE
+POSSIBILITY OF SUCH DAMAGE.
+*/
 
 #pragma once
 
diff --git a/Engine/Plugins/Runtime/Steam/SteamVR/Source/SteamVREditor/SteamVREditor.Build.cs b/Engine/Plugins/Runtime/Steam/SteamVR/Source/SteamVREditor/SteamVREditor.Build.cs
--- a/Engine/Plugins/Runtime/Steam/SteamVR/Source/SteamVREditor/SteamVREditor.Build.cs
+++ b/Engine/Plugins/Runtime/Steam/SteamVR/Source/SteamVREditor/SteamVREditor.Build.cs
@@ -1,4 +1,32 @@
-// Copyright 1998-2019 Epic Games, Inc. All Rights Reserved.
+/*
+Copyright 2019 Valve Corporation under https://opensource.org/licenses/BSD-3-Clause
+This code includes modifications by Epic Games.  Modifications (c) 2019 Epic Games, Inc.
+
+Redistribution and use in source and binary forms, with or without modification,
+are permitted provided that the following conditions are met:
+
+1. Redistributions of source code must retain the above copyright notice, this
+   list of conditions and the following disclaimer.
+
+2. Redistributions in binary form must reproduce the above copyright notice,
+   this list of conditions and the following disclaimer in the documentation
+   and/or other materials provided with the distribution.
+
+3. Neither the name of the copyright holder nor the names of its contributors
+   may be used to endorse or promote products derived from this software
+   without specific prior written permission.
+
+THIS SOFTWARE IS PROVIDED BY THE COPYRIGHT HOLDERS AND CONTRIBUTORS "AS IS" AND
+ANY EXPRESS OR IMPLIED WARRANTIES, INCLUDING, BUT NOT LIMITED TO, THE IMPLIED
+WARRANTIES OF MERCHANTABILITY AND FITNESS FOR A PARTICULAR PURPOSE ARE DISCLAIMED.
+IN NO EVENT SHALL THE COPYRIGHT HOLDER OR CONTRIBUTORS BE LIABLE FOR ANY DIRECT,
+INDIRECT, INCIDENTAL, SPECIAL, EXEMPLARY, OR CONSEQUENTIAL DAMAGES (INCLUDING,
+BUT NOT LIMITED TO, PROCUREMENT OF SUBSTITUTE GOODS OR SERVICES; LOSS OF USE, DATA,
+OR PROFITS; OR BUSINESS INTERRUPTION) HOWEVER CAUSED AND ON ANY THEORY OF LIABILITY,
+WHETHER IN CONTRACT, STRICT LIABILITY, OR TORT (INCLUDING NEGLIGENCE OR OTHERWISE)
+ARISING IN ANY WAY OUT OF THE USE OF THIS SOFTWARE, EVEN IF ADVISED OF THE
+POSSIBILITY OF SUCH DAMAGE.
+*/
 
 using UnrealBuildTool;
 using System.IO;
@@ -25,7 +53,7 @@ public class SteamVREditor : ModuleRules
                 "AnimGraph"
             }
 			);
-			
+
 		PrivateDependencyModuleNames.AddRange(
 			new string[]
 			{
@@ -39,8 +67,8 @@ public class SteamVREditor : ModuleRules
 				"SlateCore"
             }
 			);
-		
-		
+
+
 		DynamicallyLoadedModuleNames.AddRange(
 			new string[]
 			{
diff --git a/Engine/Plugins/Runtime/Steam/SteamVR/Source/SteamVRInput/Private/SteamVRInput.cpp b/Engine/Plugins/Runtime/Steam/SteamVR/Source/SteamVRInput/Private/SteamVRInput.cpp
--- a/Engine/Plugins/Runtime/Steam/SteamVR/Source/SteamVRInput/Private/SteamVRInput.cpp
+++ b/Engine/Plugins/Runtime/Steam/SteamVR/Source/SteamVRInput/Private/SteamVRInput.cpp
@@ -1,64 +1,50 @@
-// Copyright 1998-2019 Epic Games, Inc. All Rights Reserved.
+/*
+Copyright 2019 Valve Corporation under https://opensource.org/licenses/BSD-3-Clause
+This code includes modifications by Epic Games.  Modifications (c) 2019 Epic Games, Inc.
+
+Redistribution and use in source and binary forms, with or without modification,
+are permitted provided that the following conditions are met:
+
+1. Redistributions of source code must retain the above copyright notice, this
+   list of conditions and the following disclaimer.
+
+2. Redistributions in binary form must reproduce the above copyright notice,
+   this list of conditions and the following disclaimer in the documentation
+   and/or other materials provided with the distribution.
+
+3. Neither the name of the copyright holder nor the names of its contributors
+   may be used to endorse or promote products derived from this software
+   without specific prior written permission.
+
+THIS SOFTWARE IS PROVIDED BY THE COPYRIGHT HOLDERS AND CONTRIBUTORS "AS IS" AND
+ANY EXPRESS OR IMPLIED WARRANTIES, INCLUDING, BUT NOT LIMITED TO, THE IMPLIED
+WARRANTIES OF MERCHANTABILITY AND FITNESS FOR A PARTICULAR PURPOSE ARE DISCLAIMED.
+IN NO EVENT SHALL THE COPYRIGHT HOLDER OR CONTRIBUTORS BE LIABLE FOR ANY DIRECT,
+INDIRECT, INCIDENTAL, SPECIAL, EXEMPLARY, OR CONSEQUENTIAL DAMAGES (INCLUDING,
+BUT NOT LIMITED TO, PROCUREMENT OF SUBSTITUTE GOODS OR SERVICES; LOSS OF USE, DATA,
+OR PROFITS; OR BUSINESS INTERRUPTION) HOWEVER CAUSED AND ON ANY THEORY OF LIABILITY,
+WHETHER IN CONTRACT, STRICT LIABILITY, OR TORT (INCLUDING NEGLIGENCE OR OTHERWISE)
+ARISING IN ANY WAY OUT OF THE USE OF THIS SOFTWARE, EVEN IF ADVISED OF THE
+POSSIBILITY OF SUCH DAMAGE.
+*/
 
 #include "SteamVRInput.h"
-#include "SteamVRInputSettings.h"
 #include "CoreMinimal.h"
 #include "Runtime/Core/Public/Misc/Paths.h"
 #include "openvr.h"
 #include "Modules/ModuleManager.h"
 #include "Interfaces/IPluginManager.h"
 
-#if WITH_EDITOR
-#include "Editor.h"
-#include "ISettingsModule.h"
-#include "ISettingsSection.h"
-#endif
-
 #define LOCTEXT_NAMESPACE "FSteamVRInputModule"
 
 DEFINE_LOG_CATEGORY_STATIC(LogSteamInput, Log, All);
 
-#if WITH_EDITOR
-void FSteamVRInputModule::AddEditorSettings()
-{
-	ISettingsModule* SettingsModule = FModuleManager::GetModulePtr<ISettingsModule>("Settings");
-
-	// While this should usually be true, it's not guaranteed that the settings module will be loaded in the editor.
-	// UBT allows setting bBuildDeveloperTools to false while bBuildEditor can be true.
-	// The former option indirectly controls loading of the "Settings" module.
-	if (SettingsModule)
-	{
-		SettingsModule->RegisterSettings("Project", "Plugins", "SteamVR Input",
-			LOCTEXT("SteamVRInputSettingsName", "SteamVR Input"),
-			LOCTEXT("SteamVRInputSettingsDescription", "Configure SteamVR input support."),
-			GetMutableDefault<USteamVRInputSettings>()
-		);
-	}
-}
-
-void FSteamVRInputModule::RemoveEditorSettings()
-{
-	ISettingsModule* SettingsModule = FModuleManager::GetModulePtr<ISettingsModule>("Settings");
-	if (SettingsModule)
-	{
-		SettingsModule->UnregisterSettings("Project", "Plugins", "SteamVR Input");
-	}
-}
-#endif
-
 void FSteamVRInputModule::StartupModule()
 {
-#if WITH_EDITOR
-	// We don't quite have control of when the "Settings" module is loaded, so we'll wait until PostEngineInit to register settings.
-	FCoreDelegates::OnPostEngineInit.AddRaw(this, &FSteamVRInputModule::AddEditorSettings);
-#endif // WITH_EDITOR
 }
 
 void FSteamVRInputModule::ShutdownModule()
 {
-#if WITH_EDITOR
-	RemoveEditorSettings();
-#endif
 }
 
 #undef LOCTEXT_NAMESPACE
diff --git a/Engine/Plugins/Runtime/Steam/SteamVR/Source/SteamVRInput/Public/SteamVRInput.h b/Engine/Plugins/Runtime/Steam/SteamVR/Source/SteamVRInput/Public/SteamVRInput.h
--- a/Engine/Plugins/Runtime/Steam/SteamVR/Source/SteamVRInput/Public/SteamVRInput.h
+++ b/Engine/Plugins/Runtime/Steam/SteamVR/Source/SteamVRInput/Public/SteamVRInput.h
@@ -1,4 +1,32 @@
-// Copyright 1998-2019 Epic Games, Inc. All Rights Reserved.
+/*
+Copyright 2019 Valve Corporation under https://opensource.org/licenses/BSD-3-Clause
+This code includes modifications by Epic Games.  Modifications (c) 2019 Epic Games, Inc.
+
+Redistribution and use in source and binary forms, with or without modification,
+are permitted provided that the following conditions are met:
+
+1. Redistributions of source code must retain the above copyright notice, this
+   list of conditions and the following disclaimer.
+
+2. Redistributions in binary form must reproduce the above copyright notice,
+   this list of conditions and the following disclaimer in the documentation
+   and/or other materials provided with the distribution.
+
+3. Neither the name of the copyright holder nor the names of its contributors
+   may be used to endorse or promote products derived from this software
+   without specific prior written permission.
+
+THIS SOFTWARE IS PROVIDED BY THE COPYRIGHT HOLDERS AND CONTRIBUTORS "AS IS" AND
+ANY EXPRESS OR IMPLIED WARRANTIES, INCLUDING, BUT NOT LIMITED TO, THE IMPLIED
+WARRANTIES OF MERCHANTABILITY AND FITNESS FOR A PARTICULAR PURPOSE ARE DISCLAIMED.
+IN NO EVENT SHALL THE COPYRIGHT HOLDER OR CONTRIBUTORS BE LIABLE FOR ANY DIRECT,
+INDIRECT, INCIDENTAL, SPECIAL, EXEMPLARY, OR CONSEQUENTIAL DAMAGES (INCLUDING,
+BUT NOT LIMITED TO, PROCUREMENT OF SUBSTITUTE GOODS OR SERVICES; LOSS OF USE, DATA,
+OR PROFITS; OR BUSINESS INTERRUPTION) HOWEVER CAUSED AND ON ANY THEORY OF LIABILITY,
+WHETHER IN CONTRACT, STRICT LIABILITY, OR TORT (INCLUDING NEGLIGENCE OR OTHERWISE)
+ARISING IN ANY WAY OUT OF THE USE OF THIS SOFTWARE, EVEN IF ADVISED OF THE
+POSSIBILITY OF SUCH DAMAGE.
+*/
 
 #pragma once
 
@@ -12,11 +40,6 @@ public:
 	virtual void StartupModule() override;
 	virtual void ShutdownModule() override;
 
-#if WITH_EDITOR
-	void AddEditorSettings();
-	void RemoveEditorSettings();
-#endif
-
 	static inline IModuleInterface& Get()
 	{
 		return FModuleManager::LoadModuleChecked<IModuleInterface>("SteamVRInput");
diff --git a/Engine/Plugins/Runtime/Steam/SteamVR/Source/SteamVRInput/Public/SteamVRInputSettings.h b/Engine/Plugins/Runtime/Steam/SteamVR/Source/SteamVRInput/Public/SteamVRInputSettings.h
deleted file mode 100644
diff --git a/Engine/Plugins/Runtime/Steam/SteamVR/Source/SteamVRInput/SteamVRInput.Build.cs b/Engine/Plugins/Runtime/Steam/SteamVR/Source/SteamVRInput/SteamVRInput.Build.cs
--- a/Engine/Plugins/Runtime/Steam/SteamVR/Source/SteamVRInput/SteamVRInput.Build.cs
+++ b/Engine/Plugins/Runtime/Steam/SteamVR/Source/SteamVRInput/SteamVRInput.Build.cs
@@ -1,4 +1,32 @@
-// Copyright 1998-2019 Epic Games, Inc. All Rights Reserved.
+/*
+Copyright 2019 Valve Corporation under https://opensource.org/licenses/BSD-3-Clause
+This code includes modifications by Epic Games.  Modifications (c) 2019 Epic Games, Inc.
+
+Redistribution and use in source and binary forms, with or without modification,
+are permitted provided that the following conditions are met:
+
+1. Redistributions of source code must retain the above copyright notice, this
+   list of conditions and the following disclaimer.
+
+2. Redistributions in binary form must reproduce the above copyright notice,
+   this list of conditions and the following disclaimer in the documentation
+   and/or other materials provided with the distribution.
+
+3. Neither the name of the copyright holder nor the names of its contributors
+   may be used to endorse or promote products derived from this software
+   without specific prior written permission.
+
+THIS SOFTWARE IS PROVIDED BY THE COPYRIGHT HOLDERS AND CONTRIBUTORS "AS IS" AND
+ANY EXPRESS OR IMPLIED WARRANTIES, INCLUDING, BUT NOT LIMITED TO, THE IMPLIED
+WARRANTIES OF MERCHANTABILITY AND FITNESS FOR A PARTICULAR PURPOSE ARE DISCLAIMED.
+IN NO EVENT SHALL THE COPYRIGHT HOLDER OR CONTRIBUTORS BE LIABLE FOR ANY DIRECT,
+INDIRECT, INCIDENTAL, SPECIAL, EXEMPLARY, OR CONSEQUENTIAL DAMAGES (INCLUDING,
+BUT NOT LIMITED TO, PROCUREMENT OF SUBSTITUTE GOODS OR SERVICES; LOSS OF USE, DATA,
+OR PROFITS; OR BUSINESS INTERRUPTION) HOWEVER CAUSED AND ON ANY THEORY OF LIABILITY,
+WHETHER IN CONTRACT, STRICT LIABILITY, OR TORT (INCLUDING NEGLIGENCE OR OTHERWISE)
+ARISING IN ANY WAY OUT OF THE USE OF THIS SOFTWARE, EVEN IF ADVISED OF THE
+POSSIBILITY OF SUCH DAMAGE.
+*/
 
 using UnrealBuildTool;
 using System.IO;
diff --git a/Engine/Plugins/Runtime/Steam/SteamVR/Source/SteamVRInputDevice/Private/AnimNode_SteamVRInputAnimPose.cpp b/Engine/Plugins/Runtime/Steam/SteamVR/Source/SteamVRInputDevice/Private/AnimNode_SteamVRInputAnimPose.cpp
--- a/Engine/Plugins/Runtime/Steam/SteamVR/Source/SteamVRInputDevice/Private/AnimNode_SteamVRInputAnimPose.cpp
+++ b/Engine/Plugins/Runtime/Steam/SteamVR/Source/SteamVRInputDevice/Private/AnimNode_SteamVRInputAnimPose.cpp
@@ -1,4 +1,32 @@
-// Copyright 1998-2019 Epic Games, Inc. All Rights Reserved.
+/*
+Copyright 2019 Valve Corporation under https://opensource.org/licenses/BSD-3-Clause
+This code includes modifications by Epic Games.  Modifications (c) 2019 Epic Games, Inc.
+
+Redistribution and use in source and binary forms, with or without modification,
+are permitted provided that the following conditions are met:
+
+1. Redistributions of source code must retain the above copyright notice, this
+   list of conditions and the following disclaimer.
+
+2. Redistributions in binary form must reproduce the above copyright notice,
+   this list of conditions and the following disclaimer in the documentation
+   and/or other materials provided with the distribution.
+
+3. Neither the name of the copyright holder nor the names of its contributors
+   may be used to endorse or promote products derived from this software
+   without specific prior written permission.
+
+THIS SOFTWARE IS PROVIDED BY THE COPYRIGHT HOLDERS AND CONTRIBUTORS "AS IS" AND
+ANY EXPRESS OR IMPLIED WARRANTIES, INCLUDING, BUT NOT LIMITED TO, THE IMPLIED
+WARRANTIES OF MERCHANTABILITY AND FITNESS FOR A PARTICULAR PURPOSE ARE DISCLAIMED.
+IN NO EVENT SHALL THE COPYRIGHT HOLDER OR CONTRIBUTORS BE LIABLE FOR ANY DIRECT,
+INDIRECT, INCIDENTAL, SPECIAL, EXEMPLARY, OR CONSEQUENTIAL DAMAGES (INCLUDING,
+BUT NOT LIMITED TO, PROCUREMENT OF SUBSTITUTE GOODS OR SERVICES; LOSS OF USE, DATA,
+OR PROFITS; OR BUSINESS INTERRUPTION) HOWEVER CAUSED AND ON ANY THEORY OF LIABILITY,
+WHETHER IN CONTRACT, STRICT LIABILITY, OR TORT (INCLUDING NEGLIGENCE OR OTHERWISE)
+ARISING IN ANY WAY OUT OF THE USE OF THIS SOFTWARE, EVEN IF ADVISED OF THE
+POSSIBILITY OF SUCH DAMAGE.
+*/
 
 #include "AnimNode_SteamVRInputAnimPose.h"
 #include "ISteamVRInputDeviceModule.h"
diff --git a/Engine/Plugins/Runtime/Steam/SteamVR/Source/SteamVRInputDevice/Private/AnimNode_SteamVRSetWristTransform.cpp b/Engine/Plugins/Runtime/Steam/SteamVR/Source/SteamVRInputDevice/Private/AnimNode_SteamVRSetWristTransform.cpp
--- a/Engine/Plugins/Runtime/Steam/SteamVR/Source/SteamVRInputDevice/Private/AnimNode_SteamVRSetWristTransform.cpp
+++ b/Engine/Plugins/Runtime/Steam/SteamVR/Source/SteamVRInputDevice/Private/AnimNode_SteamVRSetWristTransform.cpp
@@ -1,4 +1,32 @@
-// Copyright 1998-2019 Epic Games, Inc. All Rights Reserved.
+/*
+Copyright 2019 Valve Corporation under https://opensource.org/licenses/BSD-3-Clause
+This code includes modifications by Epic Games.  Modifications (c) 2019 Epic Games, Inc.
+
+Redistribution and use in source and binary forms, with or without modification,
+are permitted provided that the following conditions are met:
+
+1. Redistributions of source code must retain the above copyright notice, this
+   list of conditions and the following disclaimer.
+
+2. Redistributions in binary form must reproduce the above copyright notice,
+   this list of conditions and the following disclaimer in the documentation
+   and/or other materials provided with the distribution.
+
+3. Neither the name of the copyright holder nor the names of its contributors
+   may be used to endorse or promote products derived from this software
+   without specific prior written permission.
+
+THIS SOFTWARE IS PROVIDED BY THE COPYRIGHT HOLDERS AND CONTRIBUTORS "AS IS" AND
+ANY EXPRESS OR IMPLIED WARRANTIES, INCLUDING, BUT NOT LIMITED TO, THE IMPLIED
+WARRANTIES OF MERCHANTABILITY AND FITNESS FOR A PARTICULAR PURPOSE ARE DISCLAIMED.
+IN NO EVENT SHALL THE COPYRIGHT HOLDER OR CONTRIBUTORS BE LIABLE FOR ANY DIRECT,
+INDIRECT, INCIDENTAL, SPECIAL, EXEMPLARY, OR CONSEQUENTIAL DAMAGES (INCLUDING,
+BUT NOT LIMITED TO, PROCUREMENT OF SUBSTITUTE GOODS OR SERVICES; LOSS OF USE, DATA,
+OR PROFITS; OR BUSINESS INTERRUPTION) HOWEVER CAUSED AND ON ANY THEORY OF LIABILITY,
+WHETHER IN CONTRACT, STRICT LIABILITY, OR TORT (INCLUDING NEGLIGENCE OR OTHERWISE)
+ARISING IN ANY WAY OUT OF THE USE OF THIS SOFTWARE, EVEN IF ADVISED OF THE
+POSSIBILITY OF SUCH DAMAGE.
+*/
 
 #include "AnimNode_SteamVRSetWristTransform.h"
 #include "ISteamVRInputDeviceModule.h"
diff --git a/Engine/Plugins/Runtime/Steam/SteamVR/Source/SteamVRInputDevice/Private/SteamVRInputDevice.cpp b/Engine/Plugins/Runtime/Steam/SteamVR/Source/SteamVRInputDevice/Private/SteamVRInputDevice.cpp
--- a/Engine/Plugins/Runtime/Steam/SteamVR/Source/SteamVRInputDevice/Private/SteamVRInputDevice.cpp
+++ b/Engine/Plugins/Runtime/Steam/SteamVR/Source/SteamVRInputDevice/Private/SteamVRInputDevice.cpp
@@ -1,7 +1,34 @@
-// Copyright 1998-2019 Epic Games, Inc. All Rights Reserved.
+/*
+Copyright 2019 Valve Corporation under https://opensource.org/licenses/BSD-3-Clause
+This code includes modifications by Epic Games.  Modifications (c) 2019 Epic Games, Inc.
+
+Redistribution and use in source and binary forms, with or without modification,
+are permitted provided that the following conditions are met:
+
+1. Redistributions of source code must retain the above copyright notice, this
+   list of conditions and the following disclaimer.
+
+2. Redistributions in binary form must reproduce the above copyright notice,
+   this list of conditions and the following disclaimer in the documentation
+   and/or other materials provided with the distribution.
+
+3. Neither the name of the copyright holder nor the names of its contributors
+   may be used to endorse or promote products derived from this software
+   without specific prior written permission.
+
+THIS SOFTWARE IS PROVIDED BY THE COPYRIGHT HOLDERS AND CONTRIBUTORS "AS IS" AND
+ANY EXPRESS OR IMPLIED WARRANTIES, INCLUDING, BUT NOT LIMITED TO, THE IMPLIED
+WARRANTIES OF MERCHANTABILITY AND FITNESS FOR A PARTICULAR PURPOSE ARE DISCLAIMED.
+IN NO EVENT SHALL THE COPYRIGHT HOLDER OR CONTRIBUTORS BE LIABLE FOR ANY DIRECT,
+INDIRECT, INCIDENTAL, SPECIAL, EXEMPLARY, OR CONSEQUENTIAL DAMAGES (INCLUDING,
+BUT NOT LIMITED TO, PROCUREMENT OF SUBSTITUTE GOODS OR SERVICES; LOSS OF USE, DATA,
+OR PROFITS; OR BUSINESS INTERRUPTION) HOWEVER CAUSED AND ON ANY THEORY OF LIABILITY,
+WHETHER IN CONTRACT, STRICT LIABILITY, OR TORT (INCLUDING NEGLIGENCE OR OTHERWISE)
+ARISING IN ANY WAY OUT OF THE USE OF THIS SOFTWARE, EVEN IF ADVISED OF THE
+POSSIBILITY OF SUCH DAMAGE.
+*/
 
 #include "SteamVRInputDevice.h"
-#include "SteamVRInputSettings.h"
 #include "Runtime/ApplicationCore/Public/GenericPlatform/IInputInterface.h"
 #include "HAL/FileManagerGeneric.h"
 #include "Misc/FileHelper.h"
@@ -16,6 +43,7 @@
 #include "SteamVRControllerKeys.h"
 #include "Runtime/HeadMountedDisplay/Public/IXRTrackingSystem.h"
 #include "SteamVRSkeletonDefinition.h"
+#include "Misc/MessageDialog.h"
 
 #if PLATFORM_WINDOWS
 #include "Windows/WindowsHWrapper.h"
@@ -83,9 +111,7 @@ FSteamVRInputDevice::FSteamVRInputDevice(const TSharedRef<FGenericApplicationMes
 #if WITH_EDITOR
 	GenerateActionManifest();
 
-	// Auto-enable SteamVR Input Developer Mode (reload hmd module)
-	const USteamVRInputSettings* Settings = GetDefault<USteamVRInputSettings>();
-	if (Settings->bDeveloperMode && VRSettings() != nullptr)
+	if (VRSettings() != nullptr)
 	{
 		EVRInitError SteamVRInitError = VRInitError_Driver_NotLoaded;
 		IVRSystem* SteamVRSystem = vr::VR_Init(&SteamVRInitError, vr::VRApplication_Overlay);
@@ -1212,11 +1238,7 @@ void FSteamVRInputDevice::RegenerateControllerBindings()
 
 void FSteamVRInputDevice::OnActionMappingsChanged()
 {
-	const USteamVRInputSettings* Settings = GetDefault<USteamVRInputSettings>();
-	if (!Settings->bDeveloperMode)
-	{
-		this->GenerateActionManifest(true, true, true, true);
-	}
+	this->GenerateActionManifest(true, true, true, true);
 }
 
 bool FSteamVRInputDevice::GenerateAppManifest(FString ManifestPath, FString ProjectName, FString& OutAppKey, FString& OutAppManifestPath)
@@ -1318,12 +1340,13 @@ void FSteamVRInputDevice::GenerateControllerBindings(const FString& BindingsPath
 	for (auto& SupportedController : InOutControllerTypes)
 	{
 		// If there is no user-defined controller binding or it hasn't been auto-generated yet, generate it
-		if (!SupportedController.bIsGenerated || bDeleteIfExists)
+		if (!SupportedController.bIsGenerated)
 		{
 			// Creating bindings file
 			TSharedRef<FJsonObject> BindingsObject = MakeShareable(new FJsonObject());
 			BindingsObject->SetStringField(TEXT("name"), TEXT("Default bindings for ") + SupportedController.Description);
 			BindingsObject->SetStringField(TEXT("controller_type"), SupportedController.Name.ToString());
+			BindingsObject->SetStringField(TEXT("last_edited_by"), FApp::GetEpicProductIdentifier());
 
 			// Create Action Bindings in JSON Format
 			TArray<TSharedPtr<FJsonValue>> JsonValuesArray;
@@ -2579,6 +2602,7 @@ void FSteamVRInputDevice::GenerateActionManifest(bool GenerateActions, bool Gene
 	UE_LOG(LogSteamVRInputDevice, Log, TEXT("Searching for Controller Bindings files at: %s"), *ControllerBindingsPath);
 
 	// Look for existing controller binding files
+	uint32 YesNoToAll = EAppReturnType::No;
 	for (FString& BindingFile : ControllerBindingFiles)
 	{
 		// Skip if manifest
@@ -2590,9 +2614,11 @@ void FSteamVRInputDevice::GenerateActionManifest(bool GenerateActions, bool Gene
 		// Setup cache
 		FString StringCache;
 		FString ControllerType;
+		FString LastEdited;
 		
 		// Load Binding File to a string
-		FFileHelper::LoadFileToString(StringCache, *(ControllerBindingsPath / BindingFile));
+		FString BindingPath = ControllerBindingsPath / BindingFile;
+		FFileHelper::LoadFileToString(StringCache, *BindingPath);
 
 		// Convert string to json object
 		TSharedRef<TJsonReader<TCHAR>> JsonReader = TJsonReaderFactory<TCHAR>::Create(StringCache);
@@ -2610,6 +2636,40 @@ void FSteamVRInputDevice::GenerateActionManifest(bool GenerateActions, bool Gene
 		}
 		else
 		{
+			bool bIsGenerated = true;
+
+			if (DeleteIfExists)
+			{
+				// Set this binding as one we need to regenerate
+				bIsGenerated = false;
+
+				// Throw up a prompt if we did not edit the binding last so the user can abort
+				if (!JsonObject->TryGetStringField(TEXT("last_edited_by"), LastEdited) || LastEdited != FApp::GetEpicProductIdentifier())
+				{
+					if (LastEdited.IsEmpty())
+					{
+						LastEdited = TEXT("SteamVR");
+					}
+
+					if (YesNoToAll != EAppReturnType::NoAll && YesNoToAll != EAppReturnType::YesAll)
+					{
+						YesNoToAll = FMessageDialog::Open(EAppMsgType::YesNoYesAllNoAll, FText::Format(LOCTEXT("BindingFileAlreadyExists", "Your binding file ({0}) was last edited by {1} do you want to overwrite the changes? You will lose any changes you made outside of the editor!"), FText::FromString(BindingFile), FText::FromString(LastEdited)));
+					}
+
+					if (YesNoToAll == EAppReturnType::No || YesNoToAll == EAppReturnType::NoAll)
+					{
+						// Prevent this binding from being regenerated
+						bIsGenerated = true;
+					}
+					else
+					{
+						FString BackupPath = BindingPath + ".backup";
+						FPlatformFileManager::Get().GetPlatformFile().DeleteFile(*BackupPath);
+						FPlatformFileManager::Get().GetPlatformFile().CopyFile(*BackupPath, *BindingPath);
+					}
+				}
+			}
+
 			// Create Controller Binding Object for this binding file
 			TSharedRef<FJsonObject> ControllerBindingObject = MakeShareable(new FJsonObject());
 			TArray<FString> ControllerStringFields = { "controller_type", *ControllerType,
@@ -2623,7 +2683,7 @@ void FSteamVRInputDevice::GenerateActionManifest(bool GenerateActions, bool Gene
 			{
 				if (DefaultControllerType.Name == FName(*ControllerType))
 				{
-					DefaultControllerType.bIsGenerated = true;
+					DefaultControllerType.bIsGenerated = bIsGenerated;
 				}
 			}
 		}
diff --git a/Engine/Plugins/Runtime/Steam/SteamVR/Source/SteamVRInputDevice/Private/SteamVRInputDeviceFunctionLibrary.cpp b/Engine/Plugins/Runtime/Steam/SteamVR/Source/SteamVRInputDevice/Private/SteamVRInputDeviceFunctionLibrary.cpp
--- a/Engine/Plugins/Runtime/Steam/SteamVR/Source/SteamVRInputDevice/Private/SteamVRInputDeviceFunctionLibrary.cpp
+++ b/Engine/Plugins/Runtime/Steam/SteamVR/Source/SteamVRInputDevice/Private/SteamVRInputDeviceFunctionLibrary.cpp
@@ -1,4 +1,32 @@
-// Copyright 1998-2019 Epic Games, Inc. All Rights Reserved.
+/*
+Copyright 2019 Valve Corporation under https://opensource.org/licenses/BSD-3-Clause
+This code includes modifications by Epic Games.  Modifications (c) 2019 Epic Games, Inc.
+
+Redistribution and use in source and binary forms, with or without modification,
+are permitted provided that the following conditions are met:
+
+1. Redistributions of source code must retain the above copyright notice, this
+   list of conditions and the following disclaimer.
+
+2. Redistributions in binary form must reproduce the above copyright notice,
+   this list of conditions and the following disclaimer in the documentation
+   and/or other materials provided with the distribution.
+
+3. Neither the name of the copyright holder nor the names of its contributors
+   may be used to endorse or promote products derived from this software
+   without specific prior written permission.
+
+THIS SOFTWARE IS PROVIDED BY THE COPYRIGHT HOLDERS AND CONTRIBUTORS "AS IS" AND
+ANY EXPRESS OR IMPLIED WARRANTIES, INCLUDING, BUT NOT LIMITED TO, THE IMPLIED
+WARRANTIES OF MERCHANTABILITY AND FITNESS FOR A PARTICULAR PURPOSE ARE DISCLAIMED.
+IN NO EVENT SHALL THE COPYRIGHT HOLDER OR CONTRIBUTORS BE LIABLE FOR ANY DIRECT,
+INDIRECT, INCIDENTAL, SPECIAL, EXEMPLARY, OR CONSEQUENTIAL DAMAGES (INCLUDING,
+BUT NOT LIMITED TO, PROCUREMENT OF SUBSTITUTE GOODS OR SERVICES; LOSS OF USE, DATA,
+OR PROFITS; OR BUSINESS INTERRUPTION) HOWEVER CAUSED AND ON ANY THEORY OF LIABILITY,
+WHETHER IN CONTRACT, STRICT LIABILITY, OR TORT (INCLUDING NEGLIGENCE OR OTHERWISE)
+ARISING IN ANY WAY OUT OF THE USE OF THIS SOFTWARE, EVEN IF ADVISED OF THE
+POSSIBILITY OF SUCH DAMAGE.
+*/
 
 #include "SteamVRInputDeviceFunctionLibrary.h"
 
diff --git a/Engine/Plugins/Runtime/Steam/SteamVR/Source/SteamVRInputDevice/Private/SteamVRInputDeviceModule.cpp b/Engine/Plugins/Runtime/Steam/SteamVR/Source/SteamVRInputDevice/Private/SteamVRInputDeviceModule.cpp
--- a/Engine/Plugins/Runtime/Steam/SteamVR/Source/SteamVRInputDevice/Private/SteamVRInputDeviceModule.cpp
+++ b/Engine/Plugins/Runtime/Steam/SteamVR/Source/SteamVRInputDevice/Private/SteamVRInputDeviceModule.cpp
@@ -1,4 +1,32 @@
-// Copyright 1998-2019 Epic Games, Inc. All Rights Reserved.
+/*
+Copyright 2019 Valve Corporation under https://opensource.org/licenses/BSD-3-Clause
+This code includes modifications by Epic Games.  Modifications (c) 2019 Epic Games, Inc.
+
+Redistribution and use in source and binary forms, with or without modification,
+are permitted provided that the following conditions are met:
+
+1. Redistributions of source code must retain the above copyright notice, this
+   list of conditions and the following disclaimer.
+
+2. Redistributions in binary form must reproduce the above copyright notice,
+   this list of conditions and the following disclaimer in the documentation
+   and/or other materials provided with the distribution.
+
+3. Neither the name of the copyright holder nor the names of its contributors
+   may be used to endorse or promote products derived from this software
+   without specific prior written permission.
+
+THIS SOFTWARE IS PROVIDED BY THE COPYRIGHT HOLDERS AND CONTRIBUTORS "AS IS" AND
+ANY EXPRESS OR IMPLIED WARRANTIES, INCLUDING, BUT NOT LIMITED TO, THE IMPLIED
+WARRANTIES OF MERCHANTABILITY AND FITNESS FOR A PARTICULAR PURPOSE ARE DISCLAIMED.
+IN NO EVENT SHALL THE COPYRIGHT HOLDER OR CONTRIBUTORS BE LIABLE FOR ANY DIRECT,
+INDIRECT, INCIDENTAL, SPECIAL, EXEMPLARY, OR CONSEQUENTIAL DAMAGES (INCLUDING,
+BUT NOT LIMITED TO, PROCUREMENT OF SUBSTITUTE GOODS OR SERVICES; LOSS OF USE, DATA,
+OR PROFITS; OR BUSINESS INTERRUPTION) HOWEVER CAUSED AND ON ANY THEORY OF LIABILITY,
+WHETHER IN CONTRACT, STRICT LIABILITY, OR TORT (INCLUDING NEGLIGENCE OR OTHERWISE)
+ARISING IN ANY WAY OUT OF THE USE OF THIS SOFTWARE, EVEN IF ADVISED OF THE
+POSSIBILITY OF SUCH DAMAGE.
+*/
 
 #include "ISteamVRInputDeviceModule.h"
 #include "IInputDevice.h"
diff --git a/Engine/Plugins/Runtime/Steam/SteamVR/Source/SteamVRInputDevice/Private/SteamVRSkeletonDefinition.cpp b/Engine/Plugins/Runtime/Steam/SteamVR/Source/SteamVRInputDevice/Private/SteamVRSkeletonDefinition.cpp
--- a/Engine/Plugins/Runtime/Steam/SteamVR/Source/SteamVRInputDevice/Private/SteamVRSkeletonDefinition.cpp
+++ b/Engine/Plugins/Runtime/Steam/SteamVR/Source/SteamVRInputDevice/Private/SteamVRSkeletonDefinition.cpp
@@ -1,4 +1,32 @@
-// Copyright 1998-2019 Epic Games, Inc. All Rights Reserved.
+/*
+Copyright 2019 Valve Corporation under https://opensource.org/licenses/BSD-3-Clause
+This code includes modifications by Epic Games.  Modifications (c) 2019 Epic Games, Inc.
+
+Redistribution and use in source and binary forms, with or without modification,
+are permitted provided that the following conditions are met:
+
+1. Redistributions of source code must retain the above copyright notice, this
+   list of conditions and the following disclaimer.
+
+2. Redistributions in binary form must reproduce the above copyright notice,
+   this list of conditions and the following disclaimer in the documentation
+   and/or other materials provided with the distribution.
+
+3. Neither the name of the copyright holder nor the names of its contributors
+   may be used to endorse or promote products derived from this software
+   without specific prior written permission.
+
+THIS SOFTWARE IS PROVIDED BY THE COPYRIGHT HOLDERS AND CONTRIBUTORS "AS IS" AND
+ANY EXPRESS OR IMPLIED WARRANTIES, INCLUDING, BUT NOT LIMITED TO, THE IMPLIED
+WARRANTIES OF MERCHANTABILITY AND FITNESS FOR A PARTICULAR PURPOSE ARE DISCLAIMED.
+IN NO EVENT SHALL THE COPYRIGHT HOLDER OR CONTRIBUTORS BE LIABLE FOR ANY DIRECT,
+INDIRECT, INCIDENTAL, SPECIAL, EXEMPLARY, OR CONSEQUENTIAL DAMAGES (INCLUDING,
+BUT NOT LIMITED TO, PROCUREMENT OF SUBSTITUTE GOODS OR SERVICES; LOSS OF USE, DATA,
+OR PROFITS; OR BUSINESS INTERRUPTION) HOWEVER CAUSED AND ON ANY THEORY OF LIABILITY,
+WHETHER IN CONTRACT, STRICT LIABILITY, OR TORT (INCLUDING NEGLIGENCE OR OTHERWISE)
+ARISING IN ANY WAY OUT OF THE USE OF THIS SOFTWARE, EVEN IF ADVISED OF THE
+POSSIBILITY OF SUCH DAMAGE.
+*/
 
 #include "SteamVRSkeletonDefinition.h"
 
diff --git a/Engine/Plugins/Runtime/Steam/SteamVR/Source/SteamVRInputDevice/Private/SteamVRTrackingRefComponent.cpp b/Engine/Plugins/Runtime/Steam/SteamVR/Source/SteamVRInputDevice/Private/SteamVRTrackingRefComponent.cpp
--- a/Engine/Plugins/Runtime/Steam/SteamVR/Source/SteamVRInputDevice/Private/SteamVRTrackingRefComponent.cpp
+++ b/Engine/Plugins/Runtime/Steam/SteamVR/Source/SteamVRInputDevice/Private/SteamVRTrackingRefComponent.cpp
@@ -1,4 +1,32 @@
-// Copyright 1998-2019 Epic Games, Inc. All Rights Reserved.
+/*
+Copyright 2019 Valve Corporation under https://opensource.org/licenses/BSD-3-Clause
+This code includes modifications by Epic Games.  Modifications (c) 2019 Epic Games, Inc.
+
+Redistribution and use in source and binary forms, with or without modification,
+are permitted provided that the following conditions are met:
+
+1. Redistributions of source code must retain the above copyright notice, this
+   list of conditions and the following disclaimer.
+
+2. Redistributions in binary form must reproduce the above copyright notice,
+   this list of conditions and the following disclaimer in the documentation
+   and/or other materials provided with the distribution.
+
+3. Neither the name of the copyright holder nor the names of its contributors
+   may be used to endorse or promote products derived from this software
+   without specific prior written permission.
+
+THIS SOFTWARE IS PROVIDED BY THE COPYRIGHT HOLDERS AND CONTRIBUTORS "AS IS" AND
+ANY EXPRESS OR IMPLIED WARRANTIES, INCLUDING, BUT NOT LIMITED TO, THE IMPLIED
+WARRANTIES OF MERCHANTABILITY AND FITNESS FOR A PARTICULAR PURPOSE ARE DISCLAIMED.
+IN NO EVENT SHALL THE COPYRIGHT HOLDER OR CONTRIBUTORS BE LIABLE FOR ANY DIRECT,
+INDIRECT, INCIDENTAL, SPECIAL, EXEMPLARY, OR CONSEQUENTIAL DAMAGES (INCLUDING,
+BUT NOT LIMITED TO, PROCUREMENT OF SUBSTITUTE GOODS OR SERVICES; LOSS OF USE, DATA,
+OR PROFITS; OR BUSINESS INTERRUPTION) HOWEVER CAUSED AND ON ANY THEORY OF LIABILITY,
+WHETHER IN CONTRACT, STRICT LIABILITY, OR TORT (INCLUDING NEGLIGENCE OR OTHERWISE)
+ARISING IN ANY WAY OUT OF THE USE OF THIS SOFTWARE, EVEN IF ADVISED OF THE
+POSSIBILITY OF SUCH DAMAGE.
+*/
 
 #include "SteamVRTrackingRefComponent.h"
 #include "openvr.h"
diff --git a/Engine/Plugins/Runtime/Steam/SteamVR/Source/SteamVRInputDevice/Private/UE4HandSkeletonDefinition.cpp b/Engine/Plugins/Runtime/Steam/SteamVR/Source/SteamVRInputDevice/Private/UE4HandSkeletonDefinition.cpp
--- a/Engine/Plugins/Runtime/Steam/SteamVR/Source/SteamVRInputDevice/Private/UE4HandSkeletonDefinition.cpp
+++ b/Engine/Plugins/Runtime/Steam/SteamVR/Source/SteamVRInputDevice/Private/UE4HandSkeletonDefinition.cpp
@@ -1,4 +1,32 @@
-// Copyright 1998-2019 Epic Games, Inc. All Rights Reserved.
+/*
+Copyright 2019 Valve Corporation under https://opensource.org/licenses/BSD-3-Clause
+This code includes modifications by Epic Games.  Modifications (c) 2019 Epic Games, Inc.
+
+Redistribution and use in source and binary forms, with or without modification,
+are permitted provided that the following conditions are met:
+
+1. Redistributions of source code must retain the above copyright notice, this
+   list of conditions and the following disclaimer.
+
+2. Redistributions in binary form must reproduce the above copyright notice,
+   this list of conditions and the following disclaimer in the documentation
+   and/or other materials provided with the distribution.
+
+3. Neither the name of the copyright holder nor the names of its contributors
+   may be used to endorse or promote products derived from this software
+   without specific prior written permission.
+
+THIS SOFTWARE IS PROVIDED BY THE COPYRIGHT HOLDERS AND CONTRIBUTORS "AS IS" AND
+ANY EXPRESS OR IMPLIED WARRANTIES, INCLUDING, BUT NOT LIMITED TO, THE IMPLIED
+WARRANTIES OF MERCHANTABILITY AND FITNESS FOR A PARTICULAR PURPOSE ARE DISCLAIMED.
+IN NO EVENT SHALL THE COPYRIGHT HOLDER OR CONTRIBUTORS BE LIABLE FOR ANY DIRECT,
+INDIRECT, INCIDENTAL, SPECIAL, EXEMPLARY, OR CONSEQUENTIAL DAMAGES (INCLUDING,
+BUT NOT LIMITED TO, PROCUREMENT OF SUBSTITUTE GOODS OR SERVICES; LOSS OF USE, DATA,
+OR PROFITS; OR BUSINESS INTERRUPTION) HOWEVER CAUSED AND ON ANY THEORY OF LIABILITY,
+WHETHER IN CONTRACT, STRICT LIABILITY, OR TORT (INCLUDING NEGLIGENCE OR OTHERWISE)
+ARISING IN ANY WAY OUT OF THE USE OF THIS SOFTWARE, EVEN IF ADVISED OF THE
+POSSIBILITY OF SUCH DAMAGE.
+*/
 
 #include "UE4HandSkeletonDefinition.h"
 
diff --git a/Engine/Plugins/Runtime/Steam/SteamVR/Source/SteamVRInputDevice/Public/AnimNode_SteamVRInputAnimPose.h b/Engine/Plugins/Runtime/Steam/SteamVR/Source/SteamVRInputDevice/Public/AnimNode_SteamVRInputAnimPose.h
--- a/Engine/Plugins/Runtime/Steam/SteamVR/Source/SteamVRInputDevice/Public/AnimNode_SteamVRInputAnimPose.h
+++ b/Engine/Plugins/Runtime/Steam/SteamVR/Source/SteamVRInputDevice/Public/AnimNode_SteamVRInputAnimPose.h
@@ -1,4 +1,32 @@
-// Copyright 1998-2019 Epic Games, Inc. All Rights Reserved.
+/*
+Copyright 2019 Valve Corporation under https://opensource.org/licenses/BSD-3-Clause
+This code includes modifications by Epic Games.  Modifications (c) 2019 Epic Games, Inc.
+
+Redistribution and use in source and binary forms, with or without modification,
+are permitted provided that the following conditions are met:
+
+1. Redistributions of source code must retain the above copyright notice, this
+   list of conditions and the following disclaimer.
+
+2. Redistributions in binary form must reproduce the above copyright notice,
+   this list of conditions and the following disclaimer in the documentation
+   and/or other materials provided with the distribution.
+
+3. Neither the name of the copyright holder nor the names of its contributors
+   may be used to endorse or promote products derived from this software
+   without specific prior written permission.
+
+THIS SOFTWARE IS PROVIDED BY THE COPYRIGHT HOLDERS AND CONTRIBUTORS "AS IS" AND
+ANY EXPRESS OR IMPLIED WARRANTIES, INCLUDING, BUT NOT LIMITED TO, THE IMPLIED
+WARRANTIES OF MERCHANTABILITY AND FITNESS FOR A PARTICULAR PURPOSE ARE DISCLAIMED.
+IN NO EVENT SHALL THE COPYRIGHT HOLDER OR CONTRIBUTORS BE LIABLE FOR ANY DIRECT,
+INDIRECT, INCIDENTAL, SPECIAL, EXEMPLARY, OR CONSEQUENTIAL DAMAGES (INCLUDING,
+BUT NOT LIMITED TO, PROCUREMENT OF SUBSTITUTE GOODS OR SERVICES; LOSS OF USE, DATA,
+OR PROFITS; OR BUSINESS INTERRUPTION) HOWEVER CAUSED AND ON ANY THEORY OF LIABILITY,
+WHETHER IN CONTRACT, STRICT LIABILITY, OR TORT (INCLUDING NEGLIGENCE OR OTHERWISE)
+ARISING IN ANY WAY OUT OF THE USE OF THIS SOFTWARE, EVEN IF ADVISED OF THE
+POSSIBILITY OF SUCH DAMAGE.
+*/
 
 #pragma once
 
@@ -51,8 +79,8 @@ public:
 
 	FAnimNode_SteamVRInputAnimPose();
 
-	/** 
-	 * Retarget the given array of bone transforms for the SteamVR skeleton to the UE4 hand skeleton and apply it to the given FPoseContext. 
+	/**
+	 * Retarget the given array of bone transforms for the SteamVR skeleton to the UE4 hand skeleton and apply it to the given FPoseContext.
 	 * The bone transforms are in the bone's local space.  Assumes that PoseContest.Pose has already been set to its reference pose
 	*/
 	void PoseUE4HandSkeleton(FCompactPose& Pose, const FTransform* BoneTransformsLS, int32 BoneTransformCount);
diff --git a/Engine/Plugins/Runtime/Steam/SteamVR/Source/SteamVRInputDevice/Public/AnimNode_SteamVRSetWristTransform.h b/Engine/Plugins/Runtime/Steam/SteamVR/Source/SteamVRInputDevice/Public/AnimNode_SteamVRSetWristTransform.h
--- a/Engine/Plugins/Runtime/Steam/SteamVR/Source/SteamVRInputDevice/Public/AnimNode_SteamVRSetWristTransform.h
+++ b/Engine/Plugins/Runtime/Steam/SteamVR/Source/SteamVRInputDevice/Public/AnimNode_SteamVRSetWristTransform.h
@@ -1,4 +1,32 @@
-// Copyright 1998-2019 Epic Games, Inc. All Rights Reserved.
+/*
+Copyright 2019 Valve Corporation under https://opensource.org/licenses/BSD-3-Clause
+This code includes modifications by Epic Games.  Modifications (c) 2019 Epic Games, Inc.
+
+Redistribution and use in source and binary forms, with or without modification,
+are permitted provided that the following conditions are met:
+
+1. Redistributions of source code must retain the above copyright notice, this
+   list of conditions and the following disclaimer.
+
+2. Redistributions in binary form must reproduce the above copyright notice,
+   this list of conditions and the following disclaimer in the documentation
+   and/or other materials provided with the distribution.
+
+3. Neither the name of the copyright holder nor the names of its contributors
+   may be used to endorse or promote products derived from this software
+   without specific prior written permission.
+
+THIS SOFTWARE IS PROVIDED BY THE COPYRIGHT HOLDERS AND CONTRIBUTORS "AS IS" AND
+ANY EXPRESS OR IMPLIED WARRANTIES, INCLUDING, BUT NOT LIMITED TO, THE IMPLIED
+WARRANTIES OF MERCHANTABILITY AND FITNESS FOR A PARTICULAR PURPOSE ARE DISCLAIMED.
+IN NO EVENT SHALL THE COPYRIGHT HOLDER OR CONTRIBUTORS BE LIABLE FOR ANY DIRECT,
+INDIRECT, INCIDENTAL, SPECIAL, EXEMPLARY, OR CONSEQUENTIAL DAMAGES (INCLUDING,
+BUT NOT LIMITED TO, PROCUREMENT OF SUBSTITUTE GOODS OR SERVICES; LOSS OF USE, DATA,
+OR PROFITS; OR BUSINESS INTERRUPTION) HOWEVER CAUSED AND ON ANY THEORY OF LIABILITY,
+WHETHER IN CONTRACT, STRICT LIABILITY, OR TORT (INCLUDING NEGLIGENCE OR OTHERWISE)
+ARISING IN ANY WAY OUT OF THE USE OF THIS SOFTWARE, EVEN IF ADVISED OF THE
+POSSIBILITY OF SUCH DAMAGE.
+*/
 
 #pragma once
 
diff --git a/Engine/Plugins/Runtime/Steam/SteamVR/Source/SteamVRInputDevice/Public/ISteamVRInputDeviceModule.h b/Engine/Plugins/Runtime/Steam/SteamVR/Source/SteamVRInputDevice/Public/ISteamVRInputDeviceModule.h
--- a/Engine/Plugins/Runtime/Steam/SteamVR/Source/SteamVRInputDevice/Public/ISteamVRInputDeviceModule.h
+++ b/Engine/Plugins/Runtime/Steam/SteamVR/Source/SteamVRInputDevice/Public/ISteamVRInputDeviceModule.h
@@ -1,4 +1,32 @@
-// Copyright 1998-2019 Epic Games, Inc. All Rights Reserved.
+/*
+Copyright 2019 Valve Corporation under https://opensource.org/licenses/BSD-3-Clause
+This code includes modifications by Epic Games.  Modifications (c) 2019 Epic Games, Inc.
+
+Redistribution and use in source and binary forms, with or without modification,
+are permitted provided that the following conditions are met:
+
+1. Redistributions of source code must retain the above copyright notice, this
+   list of conditions and the following disclaimer.
+
+2. Redistributions in binary form must reproduce the above copyright notice,
+   this list of conditions and the following disclaimer in the documentation
+   and/or other materials provided with the distribution.
+
+3. Neither the name of the copyright holder nor the names of its contributors
+   may be used to endorse or promote products derived from this software
+   without specific prior written permission.
+
+THIS SOFTWARE IS PROVIDED BY THE COPYRIGHT HOLDERS AND CONTRIBUTORS "AS IS" AND
+ANY EXPRESS OR IMPLIED WARRANTIES, INCLUDING, BUT NOT LIMITED TO, THE IMPLIED
+WARRANTIES OF MERCHANTABILITY AND FITNESS FOR A PARTICULAR PURPOSE ARE DISCLAIMED.
+IN NO EVENT SHALL THE COPYRIGHT HOLDER OR CONTRIBUTORS BE LIABLE FOR ANY DIRECT,
+INDIRECT, INCIDENTAL, SPECIAL, EXEMPLARY, OR CONSEQUENTIAL DAMAGES (INCLUDING,
+BUT NOT LIMITED TO, PROCUREMENT OF SUBSTITUTE GOODS OR SERVICES; LOSS OF USE, DATA,
+OR PROFITS; OR BUSINESS INTERRUPTION) HOWEVER CAUSED AND ON ANY THEORY OF LIABILITY,
+WHETHER IN CONTRACT, STRICT LIABILITY, OR TORT (INCLUDING NEGLIGENCE OR OTHERWISE)
+ARISING IN ANY WAY OUT OF THE USE OF THIS SOFTWARE, EVEN IF ADVISED OF THE
+POSSIBILITY OF SUCH DAMAGE.
+*/
 
 #pragma once
 
diff --git a/Engine/Plugins/Runtime/Steam/SteamVR/Source/SteamVRInputDevice/Public/SteamVRControllerKeys.h b/Engine/Plugins/Runtime/Steam/SteamVR/Source/SteamVRInputDevice/Public/SteamVRControllerKeys.h
--- a/Engine/Plugins/Runtime/Steam/SteamVR/Source/SteamVRInputDevice/Public/SteamVRControllerKeys.h
+++ b/Engine/Plugins/Runtime/Steam/SteamVR/Source/SteamVRInputDevice/Public/SteamVRControllerKeys.h
@@ -1,4 +1,32 @@
-// Copyright 1998-2019 Epic Games, Inc. All Rights Reserved.
+/*
+Copyright 2019 Valve Corporation under https://opensource.org/licenses/BSD-3-Clause
+This code includes modifications by Epic Games.  Modifications (c) 2019 Epic Games, Inc.
+
+Redistribution and use in source and binary forms, with or without modification,
+are permitted provided that the following conditions are met:
+
+1. Redistributions of source code must retain the above copyright notice, this
+   list of conditions and the following disclaimer.
+
+2. Redistributions in binary form must reproduce the above copyright notice,
+   this list of conditions and the following disclaimer in the documentation
+   and/or other materials provided with the distribution.
+
+3. Neither the name of the copyright holder nor the names of its contributors
+   may be used to endorse or promote products derived from this software
+   without specific prior written permission.
+
+THIS SOFTWARE IS PROVIDED BY THE COPYRIGHT HOLDERS AND CONTRIBUTORS "AS IS" AND
+ANY EXPRESS OR IMPLIED WARRANTIES, INCLUDING, BUT NOT LIMITED TO, THE IMPLIED
+WARRANTIES OF MERCHANTABILITY AND FITNESS FOR A PARTICULAR PURPOSE ARE DISCLAIMED.
+IN NO EVENT SHALL THE COPYRIGHT HOLDER OR CONTRIBUTORS BE LIABLE FOR ANY DIRECT,
+INDIRECT, INCIDENTAL, SPECIAL, EXEMPLARY, OR CONSEQUENTIAL DAMAGES (INCLUDING,
+BUT NOT LIMITED TO, PROCUREMENT OF SUBSTITUTE GOODS OR SERVICES; LOSS OF USE, DATA,
+OR PROFITS; OR BUSINESS INTERRUPTION) HOWEVER CAUSED AND ON ANY THEORY OF LIABILITY,
+WHETHER IN CONTRACT, STRICT LIABILITY, OR TORT (INCLUDING NEGLIGENCE OR OTHERWISE)
+ARISING IN ANY WAY OUT OF THE USE OF THIS SOFTWARE, EVEN IF ADVISED OF THE
+POSSIBILITY OF SUCH DAMAGE.
+*/
 
 #pragma once
 #include "InputCoreTypes.h"
diff --git a/Engine/Plugins/Runtime/Steam/SteamVR/Source/SteamVRInputDevice/Public/SteamVRInputDevice.h b/Engine/Plugins/Runtime/Steam/SteamVR/Source/SteamVRInputDevice/Public/SteamVRInputDevice.h
--- a/Engine/Plugins/Runtime/Steam/SteamVR/Source/SteamVRInputDevice/Public/SteamVRInputDevice.h
+++ b/Engine/Plugins/Runtime/Steam/SteamVR/Source/SteamVRInputDevice/Public/SteamVRInputDevice.h
@@ -1,4 +1,32 @@
-// Copyright 1998-2019 Epic Games, Inc. All Rights Reserved.
+/*
+Copyright 2019 Valve Corporation under https://opensource.org/licenses/BSD-3-Clause
+This code includes modifications by Epic Games.  Modifications (c) 2019 Epic Games, Inc.
+
+Redistribution and use in source and binary forms, with or without modification,
+are permitted provided that the following conditions are met:
+
+1. Redistributions of source code must retain the above copyright notice, this
+   list of conditions and the following disclaimer.
+
+2. Redistributions in binary form must reproduce the above copyright notice,
+   this list of conditions and the following disclaimer in the documentation
+   and/or other materials provided with the distribution.
+
+3. Neither the name of the copyright holder nor the names of its contributors
+   may be used to endorse or promote products derived from this software
+   without specific prior written permission.
+
+THIS SOFTWARE IS PROVIDED BY THE COPYRIGHT HOLDERS AND CONTRIBUTORS "AS IS" AND
+ANY EXPRESS OR IMPLIED WARRANTIES, INCLUDING, BUT NOT LIMITED TO, THE IMPLIED
+WARRANTIES OF MERCHANTABILITY AND FITNESS FOR A PARTICULAR PURPOSE ARE DISCLAIMED.
+IN NO EVENT SHALL THE COPYRIGHT HOLDER OR CONTRIBUTORS BE LIABLE FOR ANY DIRECT,
+INDIRECT, INCIDENTAL, SPECIAL, EXEMPLARY, OR CONSEQUENTIAL DAMAGES (INCLUDING,
+BUT NOT LIMITED TO, PROCUREMENT OF SUBSTITUTE GOODS OR SERVICES; LOSS OF USE, DATA,
+OR PROFITS; OR BUSINESS INTERRUPTION) HOWEVER CAUSED AND ON ANY THEORY OF LIABILITY,
+WHETHER IN CONTRACT, STRICT LIABILITY, OR TORT (INCLUDING NEGLIGENCE OR OTHERWISE)
+ARISING IN ANY WAY OUT OF THE USE OF THIS SOFTWARE, EVEN IF ADVISED OF THE
+POSSIBILITY OF SUCH DAMAGE.
+*/
 
 #pragma once
 
diff --git a/Engine/Plugins/Runtime/Steam/SteamVR/Source/SteamVRInputDevice/Public/SteamVRInputDeviceFunctionLibrary.h b/Engine/Plugins/Runtime/Steam/SteamVR/Source/SteamVRInputDevice/Public/SteamVRInputDeviceFunctionLibrary.h
--- a/Engine/Plugins/Runtime/Steam/SteamVR/Source/SteamVRInputDevice/Public/SteamVRInputDeviceFunctionLibrary.h
+++ b/Engine/Plugins/Runtime/Steam/SteamVR/Source/SteamVRInputDevice/Public/SteamVRInputDeviceFunctionLibrary.h
@@ -1,4 +1,32 @@
-// Copyright 1998-2019 Epic Games, Inc. All Rights Reserved.
+/*
+Copyright 2019 Valve Corporation under https://opensource.org/licenses/BSD-3-Clause
+This code includes modifications by Epic Games.  Modifications (c) 2019 Epic Games, Inc.
+
+Redistribution and use in source and binary forms, with or without modification,
+are permitted provided that the following conditions are met:
+
+1. Redistributions of source code must retain the above copyright notice, this
+   list of conditions and the following disclaimer.
+
+2. Redistributions in binary form must reproduce the above copyright notice,
+   this list of conditions and the following disclaimer in the documentation
+   and/or other materials provided with the distribution.
+
+3. Neither the name of the copyright holder nor the names of its contributors
+   may be used to endorse or promote products derived from this software
+   without specific prior written permission.
+
+THIS SOFTWARE IS PROVIDED BY THE COPYRIGHT HOLDERS AND CONTRIBUTORS "AS IS" AND
+ANY EXPRESS OR IMPLIED WARRANTIES, INCLUDING, BUT NOT LIMITED TO, THE IMPLIED
+WARRANTIES OF MERCHANTABILITY AND FITNESS FOR A PARTICULAR PURPOSE ARE DISCLAIMED.
+IN NO EVENT SHALL THE COPYRIGHT HOLDER OR CONTRIBUTORS BE LIABLE FOR ANY DIRECT,
+INDIRECT, INCIDENTAL, SPECIAL, EXEMPLARY, OR CONSEQUENTIAL DAMAGES (INCLUDING,
+BUT NOT LIMITED TO, PROCUREMENT OF SUBSTITUTE GOODS OR SERVICES; LOSS OF USE, DATA,
+OR PROFITS; OR BUSINESS INTERRUPTION) HOWEVER CAUSED AND ON ANY THEORY OF LIABILITY,
+WHETHER IN CONTRACT, STRICT LIABILITY, OR TORT (INCLUDING NEGLIGENCE OR OTHERWISE)
+ARISING IN ANY WAY OUT OF THE USE OF THIS SOFTWARE, EVEN IF ADVISED OF THE
+POSSIBILITY OF SUCH DAMAGE.
+*/
 
 #pragma once
 
diff --git a/Engine/Plugins/Runtime/Steam/SteamVR/Source/SteamVRInputDevice/Public/SteamVRInputPublic.h b/Engine/Plugins/Runtime/Steam/SteamVR/Source/SteamVRInputDevice/Public/SteamVRInputPublic.h
--- a/Engine/Plugins/Runtime/Steam/SteamVR/Source/SteamVRInputDevice/Public/SteamVRInputPublic.h
+++ b/Engine/Plugins/Runtime/Steam/SteamVR/Source/SteamVRInputDevice/Public/SteamVRInputPublic.h
@@ -1,4 +1,32 @@
-// Copyright 1998-2019 Epic Games, Inc. All Rights Reserved.
+/*
+Copyright 2019 Valve Corporation under https://opensource.org/licenses/BSD-3-Clause
+This code includes modifications by Epic Games.  Modifications (c) 2019 Epic Games, Inc.
+
+Redistribution and use in source and binary forms, with or without modification,
+are permitted provided that the following conditions are met:
+
+1. Redistributions of source code must retain the above copyright notice, this
+   list of conditions and the following disclaimer.
+
+2. Redistributions in binary form must reproduce the above copyright notice,
+   this list of conditions and the following disclaimer in the documentation
+   and/or other materials provided with the distribution.
+
+3. Neither the name of the copyright holder nor the names of its contributors
+   may be used to endorse or promote products derived from this software
+   without specific prior written permission.
+
+THIS SOFTWARE IS PROVIDED BY THE COPYRIGHT HOLDERS AND CONTRIBUTORS "AS IS" AND
+ANY EXPRESS OR IMPLIED WARRANTIES, INCLUDING, BUT NOT LIMITED TO, THE IMPLIED
+WARRANTIES OF MERCHANTABILITY AND FITNESS FOR A PARTICULAR PURPOSE ARE DISCLAIMED.
+IN NO EVENT SHALL THE COPYRIGHT HOLDER OR CONTRIBUTORS BE LIABLE FOR ANY DIRECT,
+INDIRECT, INCIDENTAL, SPECIAL, EXEMPLARY, OR CONSEQUENTIAL DAMAGES (INCLUDING,
+BUT NOT LIMITED TO, PROCUREMENT OF SUBSTITUTE GOODS OR SERVICES; LOSS OF USE, DATA,
+OR PROFITS; OR BUSINESS INTERRUPTION) HOWEVER CAUSED AND ON ANY THEORY OF LIABILITY,
+WHETHER IN CONTRACT, STRICT LIABILITY, OR TORT (INCLUDING NEGLIGENCE OR OTHERWISE)
+ARISING IN ANY WAY OUT OF THE USE OF THIS SOFTWARE, EVEN IF ADVISED OF THE
+POSSIBILITY OF SUCH DAMAGE.
+*/
 
 struct FSteamVRSkeletonTransform_t
 {
diff --git a/Engine/Plugins/Runtime/Steam/SteamVR/Source/SteamVRInputDevice/Public/SteamVRInputTypes.h b/Engine/Plugins/Runtime/Steam/SteamVR/Source/SteamVRInputDevice/Public/SteamVRInputTypes.h
--- a/Engine/Plugins/Runtime/Steam/SteamVR/Source/SteamVRInputDevice/Public/SteamVRInputTypes.h
+++ b/Engine/Plugins/Runtime/Steam/SteamVR/Source/SteamVRInputDevice/Public/SteamVRInputTypes.h
@@ -1,4 +1,32 @@
-// Copyright 1998-2019 Epic Games, Inc. All Rights Reserved.
+/*
+Copyright 2019 Valve Corporation under https://opensource.org/licenses/BSD-3-Clause
+This code includes modifications by Epic Games.  Modifications (c) 2019 Epic Games, Inc.
+
+Redistribution and use in source and binary forms, with or without modification,
+are permitted provided that the following conditions are met:
+
+1. Redistributions of source code must retain the above copyright notice, this
+   list of conditions and the following disclaimer.
+
+2. Redistributions in binary form must reproduce the above copyright notice,
+   this list of conditions and the following disclaimer in the documentation
+   and/or other materials provided with the distribution.
+
+3. Neither the name of the copyright holder nor the names of its contributors
+   may be used to endorse or promote products derived from this software
+   without specific prior written permission.
+
+THIS SOFTWARE IS PROVIDED BY THE COPYRIGHT HOLDERS AND CONTRIBUTORS "AS IS" AND
+ANY EXPRESS OR IMPLIED WARRANTIES, INCLUDING, BUT NOT LIMITED TO, THE IMPLIED
+WARRANTIES OF MERCHANTABILITY AND FITNESS FOR A PARTICULAR PURPOSE ARE DISCLAIMED.
+IN NO EVENT SHALL THE COPYRIGHT HOLDER OR CONTRIBUTORS BE LIABLE FOR ANY DIRECT,
+INDIRECT, INCIDENTAL, SPECIAL, EXEMPLARY, OR CONSEQUENTIAL DAMAGES (INCLUDING,
+BUT NOT LIMITED TO, PROCUREMENT OF SUBSTITUTE GOODS OR SERVICES; LOSS OF USE, DATA,
+OR PROFITS; OR BUSINESS INTERRUPTION) HOWEVER CAUSED AND ON ANY THEORY OF LIABILITY,
+WHETHER IN CONTRACT, STRICT LIABILITY, OR TORT (INCLUDING NEGLIGENCE OR OTHERWISE)
+ARISING IN ANY WAY OUT OF THE USE OF THIS SOFTWARE, EVEN IF ADVISED OF THE
+POSSIBILITY OF SUCH DAMAGE.
+*/
 
 #pragma once
 
diff --git a/Engine/Plugins/Runtime/Steam/SteamVR/Source/SteamVRInputDevice/Public/SteamVRKnucklesKeys.h b/Engine/Plugins/Runtime/Steam/SteamVR/Source/SteamVRInputDevice/Public/SteamVRKnucklesKeys.h
--- a/Engine/Plugins/Runtime/Steam/SteamVR/Source/SteamVRInputDevice/Public/SteamVRKnucklesKeys.h
+++ b/Engine/Plugins/Runtime/Steam/SteamVR/Source/SteamVRInputDevice/Public/SteamVRKnucklesKeys.h
@@ -1,4 +1,32 @@
-// Copyright 1998-2019 Epic Games, Inc. All Rights Reserved.
+/*
+Copyright 2019 Valve Corporation under https://opensource.org/licenses/BSD-3-Clause
+This code includes modifications by Epic Games.  Modifications (c) 2019 Epic Games, Inc.
+
+Redistribution and use in source and binary forms, with or without modification,
+are permitted provided that the following conditions are met:
+
+1. Redistributions of source code must retain the above copyright notice, this
+   list of conditions and the following disclaimer.
+
+2. Redistributions in binary form must reproduce the above copyright notice,
+   this list of conditions and the following disclaimer in the documentation
+   and/or other materials provided with the distribution.
+
+3. Neither the name of the copyright holder nor the names of its contributors
+   may be used to endorse or promote products derived from this software
+   without specific prior written permission.
+
+THIS SOFTWARE IS PROVIDED BY THE COPYRIGHT HOLDERS AND CONTRIBUTORS "AS IS" AND
+ANY EXPRESS OR IMPLIED WARRANTIES, INCLUDING, BUT NOT LIMITED TO, THE IMPLIED
+WARRANTIES OF MERCHANTABILITY AND FITNESS FOR A PARTICULAR PURPOSE ARE DISCLAIMED.
+IN NO EVENT SHALL THE COPYRIGHT HOLDER OR CONTRIBUTORS BE LIABLE FOR ANY DIRECT,
+INDIRECT, INCIDENTAL, SPECIAL, EXEMPLARY, OR CONSEQUENTIAL DAMAGES (INCLUDING,
+BUT NOT LIMITED TO, PROCUREMENT OF SUBSTITUTE GOODS OR SERVICES; LOSS OF USE, DATA,
+OR PROFITS; OR BUSINESS INTERRUPTION) HOWEVER CAUSED AND ON ANY THEORY OF LIABILITY,
+WHETHER IN CONTRACT, STRICT LIABILITY, OR TORT (INCLUDING NEGLIGENCE OR OTHERWISE)
+ARISING IN ANY WAY OUT OF THE USE OF THIS SOFTWARE, EVEN IF ADVISED OF THE
+POSSIBILITY OF SUCH DAMAGE.
+*/
 
 #pragma once
 #include "InputCoreTypes.h"
diff --git a/Engine/Plugins/Runtime/Steam/SteamVR/Source/SteamVRInputDevice/Public/SteamVRSkeletonDefinition.h b/Engine/Plugins/Runtime/Steam/SteamVR/Source/SteamVRInputDevice/Public/SteamVRSkeletonDefinition.h
--- a/Engine/Plugins/Runtime/Steam/SteamVR/Source/SteamVRInputDevice/Public/SteamVRSkeletonDefinition.h
+++ b/Engine/Plugins/Runtime/Steam/SteamVR/Source/SteamVRInputDevice/Public/SteamVRSkeletonDefinition.h
@@ -1,4 +1,32 @@
-// Copyright 1998-2019 Epic Games, Inc. All Rights Reserved.
+/*
+Copyright 2019 Valve Corporation under https://opensource.org/licenses/BSD-3-Clause
+This code includes modifications by Epic Games.  Modifications (c) 2019 Epic Games, Inc.
+
+Redistribution and use in source and binary forms, with or without modification,
+are permitted provided that the following conditions are met:
+
+1. Redistributions of source code must retain the above copyright notice, this
+   list of conditions and the following disclaimer.
+
+2. Redistributions in binary form must reproduce the above copyright notice,
+   this list of conditions and the following disclaimer in the documentation
+   and/or other materials provided with the distribution.
+
+3. Neither the name of the copyright holder nor the names of its contributors
+   may be used to endorse or promote products derived from this software
+   without specific prior written permission.
+
+THIS SOFTWARE IS PROVIDED BY THE COPYRIGHT HOLDERS AND CONTRIBUTORS "AS IS" AND
+ANY EXPRESS OR IMPLIED WARRANTIES, INCLUDING, BUT NOT LIMITED TO, THE IMPLIED
+WARRANTIES OF MERCHANTABILITY AND FITNESS FOR A PARTICULAR PURPOSE ARE DISCLAIMED.
+IN NO EVENT SHALL THE COPYRIGHT HOLDER OR CONTRIBUTORS BE LIABLE FOR ANY DIRECT,
+INDIRECT, INCIDENTAL, SPECIAL, EXEMPLARY, OR CONSEQUENTIAL DAMAGES (INCLUDING,
+BUT NOT LIMITED TO, PROCUREMENT OF SUBSTITUTE GOODS OR SERVICES; LOSS OF USE, DATA,
+OR PROFITS; OR BUSINESS INTERRUPTION) HOWEVER CAUSED AND ON ANY THEORY OF LIABILITY,
+WHETHER IN CONTRACT, STRICT LIABILITY, OR TORT (INCLUDING NEGLIGENCE OR OTHERWISE)
+ARISING IN ANY WAY OUT OF THE USE OF THIS SOFTWARE, EVEN IF ADVISED OF THE
+POSSIBILITY OF SUCH DAMAGE.
+*/
 
 #pragma once
 
diff --git a/Engine/Plugins/Runtime/Steam/SteamVR/Source/SteamVRInputDevice/Public/SteamVRTrackingRefComponent.h b/Engine/Plugins/Runtime/Steam/SteamVR/Source/SteamVRInputDevice/Public/SteamVRTrackingRefComponent.h
--- a/Engine/Plugins/Runtime/Steam/SteamVR/Source/SteamVRInputDevice/Public/SteamVRTrackingRefComponent.h
+++ b/Engine/Plugins/Runtime/Steam/SteamVR/Source/SteamVRInputDevice/Public/SteamVRTrackingRefComponent.h
@@ -1,4 +1,32 @@
-// Copyright 1998-2019 Epic Games, Inc. All Rights Reserved.
+/*
+Copyright 2019 Valve Corporation under https://opensource.org/licenses/BSD-3-Clause
+This code includes modifications by Epic Games.  Modifications (c) 2019 Epic Games, Inc.
+
+Redistribution and use in source and binary forms, with or without modification,
+are permitted provided that the following conditions are met:
+
+1. Redistributions of source code must retain the above copyright notice, this
+   list of conditions and the following disclaimer.
+
+2. Redistributions in binary form must reproduce the above copyright notice,
+   this list of conditions and the following disclaimer in the documentation
+   and/or other materials provided with the distribution.
+
+3. Neither the name of the copyright holder nor the names of its contributors
+   may be used to endorse or promote products derived from this software
+   without specific prior written permission.
+
+THIS SOFTWARE IS PROVIDED BY THE COPYRIGHT HOLDERS AND CONTRIBUTORS "AS IS" AND
+ANY EXPRESS OR IMPLIED WARRANTIES, INCLUDING, BUT NOT LIMITED TO, THE IMPLIED
+WARRANTIES OF MERCHANTABILITY AND FITNESS FOR A PARTICULAR PURPOSE ARE DISCLAIMED.
+IN NO EVENT SHALL THE COPYRIGHT HOLDER OR CONTRIBUTORS BE LIABLE FOR ANY DIRECT,
+INDIRECT, INCIDENTAL, SPECIAL, EXEMPLARY, OR CONSEQUENTIAL DAMAGES (INCLUDING,
+BUT NOT LIMITED TO, PROCUREMENT OF SUBSTITUTE GOODS OR SERVICES; LOSS OF USE, DATA,
+OR PROFITS; OR BUSINESS INTERRUPTION) HOWEVER CAUSED AND ON ANY THEORY OF LIABILITY,
+WHETHER IN CONTRACT, STRICT LIABILITY, OR TORT (INCLUDING NEGLIGENCE OR OTHERWISE)
+ARISING IN ANY WAY OUT OF THE USE OF THIS SOFTWARE, EVEN IF ADVISED OF THE
+POSSIBILITY OF SUCH DAMAGE.
+*/
 
 #pragma once
 
diff --git a/Engine/Plugins/Runtime/Steam/SteamVR/Source/SteamVRInputDevice/Public/UE4HandSkeletonDefinition.h b/Engine/Plugins/Runtime/Steam/SteamVR/Source/SteamVRInputDevice/Public/UE4HandSkeletonDefinition.h
--- a/Engine/Plugins/Runtime/Steam/SteamVR/Source/SteamVRInputDevice/Public/UE4HandSkeletonDefinition.h
+++ b/Engine/Plugins/Runtime/Steam/SteamVR/Source/SteamVRInputDevice/Public/UE4HandSkeletonDefinition.h
@@ -1,4 +1,32 @@
-// Copyright 1998-2019 Epic Games, Inc. All Rights Reserved.
+/*
+Copyright 2019 Valve Corporation under https://opensource.org/licenses/BSD-3-Clause
+This code includes modifications by Epic Games.  Modifications (c) 2019 Epic Games, Inc.
+
+Redistribution and use in source and binary forms, with or without modification,
+are permitted provided that the following conditions are met:
+
+1. Redistributions of source code must retain the above copyright notice, this
+   list of conditions and the following disclaimer.
+
+2. Redistributions in binary form must reproduce the above copyright notice,
+   this list of conditions and the following disclaimer in the documentation
+   and/or other materials provided with the distribution.
+
+3. Neither the name of the copyright holder nor the names of its contributors
+   may be used to endorse or promote products derived from this software
+   without specific prior written permission.
+
+THIS SOFTWARE IS PROVIDED BY THE COPYRIGHT HOLDERS AND CONTRIBUTORS "AS IS" AND
+ANY EXPRESS OR IMPLIED WARRANTIES, INCLUDING, BUT NOT LIMITED TO, THE IMPLIED
+WARRANTIES OF MERCHANTABILITY AND FITNESS FOR A PARTICULAR PURPOSE ARE DISCLAIMED.
+IN NO EVENT SHALL THE COPYRIGHT HOLDER OR CONTRIBUTORS BE LIABLE FOR ANY DIRECT,
+INDIRECT, INCIDENTAL, SPECIAL, EXEMPLARY, OR CONSEQUENTIAL DAMAGES (INCLUDING,
+BUT NOT LIMITED TO, PROCUREMENT OF SUBSTITUTE GOODS OR SERVICES; LOSS OF USE, DATA,
+OR PROFITS; OR BUSINESS INTERRUPTION) HOWEVER CAUSED AND ON ANY THEORY OF LIABILITY,
+WHETHER IN CONTRACT, STRICT LIABILITY, OR TORT (INCLUDING NEGLIGENCE OR OTHERWISE)
+ARISING IN ANY WAY OUT OF THE USE OF THIS SOFTWARE, EVEN IF ADVISED OF THE
+POSSIBILITY OF SUCH DAMAGE.
+*/
 
 #pragma once
 
diff --git a/Engine/Plugins/Runtime/Steam/SteamVR/Source/SteamVRInputDevice/SteamVRInputDevice.Build.cs b/Engine/Plugins/Runtime/Steam/SteamVR/Source/SteamVRInputDevice/SteamVRInputDevice.Build.cs
--- a/Engine/Plugins/Runtime/Steam/SteamVR/Source/SteamVRInputDevice/SteamVRInputDevice.Build.cs
+++ b/Engine/Plugins/Runtime/Steam/SteamVR/Source/SteamVRInputDevice/SteamVRInputDevice.Build.cs
@@ -1,4 +1,32 @@
-// Copyright 1998-2019 Epic Games, Inc. All Rights Reserved.
+/*
+Copyright 2019 Valve Corporation under https://opensource.org/licenses/BSD-3-Clause
+This code includes modifications by Epic Games.  Modifications (c) 2019 Epic Games, Inc.
+
+Redistribution and use in source and binary forms, with or without modification,
+are permitted provided that the following conditions are met:
+
+1. Redistributions of source code must retain the above copyright notice, this
+   list of conditions and the following disclaimer.
+
+2. Redistributions in binary form must reproduce the above copyright notice,
+   this list of conditions and the following disclaimer in the documentation
+   and/or other materials provided with the distribution.
+
+3. Neither the name of the copyright holder nor the names of its contributors
+   may be used to endorse or promote products derived from this software
+   without specific prior written permission.
+
+THIS SOFTWARE IS PROVIDED BY THE COPYRIGHT HOLDERS AND CONTRIBUTORS "AS IS" AND
+ANY EXPRESS OR IMPLIED WARRANTIES, INCLUDING, BUT NOT LIMITED TO, THE IMPLIED
+WARRANTIES OF MERCHANTABILITY AND FITNESS FOR A PARTICULAR PURPOSE ARE DISCLAIMED.
+IN NO EVENT SHALL THE COPYRIGHT HOLDER OR CONTRIBUTORS BE LIABLE FOR ANY DIRECT,
+INDIRECT, INCIDENTAL, SPECIAL, EXEMPLARY, OR CONSEQUENTIAL DAMAGES (INCLUDING,
+BUT NOT LIMITED TO, PROCUREMENT OF SUBSTITUTE GOODS OR SERVICES; LOSS OF USE, DATA,
+OR PROFITS; OR BUSINESS INTERRUPTION) HOWEVER CAUSED AND ON ANY THEORY OF LIABILITY,
+WHETHER IN CONTRACT, STRICT LIABILITY, OR TORT (INCLUDING NEGLIGENCE OR OTHERWISE)
+ARISING IN ANY WAY OUT OF THE USE OF THIS SOFTWARE, EVEN IF ADVISED OF THE
+POSSIBILITY OF SUCH DAMAGE.
+*/
 
 using UnrealBuildTool;
 using System.IO;
diff --git a/Engine/Plugins/Runtime/Steam/SteamVR/SteamVR.uplugin b/Engine/Plugins/Runtime/Steam/SteamVR/SteamVR.uplugin
--- a/Engine/Plugins/Runtime/Steam/SteamVR/SteamVR.uplugin
+++ b/Engine/Plugins/Runtime/Steam/SteamVR/SteamVR.uplugin
@@ -51,7 +51,7 @@
 		{
 			"Name": "SteamVREditor",
 			"Type": "Editor",
-			"LoadingPhase": "Default",
+			"LoadingPhase": "PostEngineInit",
 			"WhitelistPlatforms": [
 				"Win64"
 			]
