Fix thumbsticks not getting bound in SteamVR Input

diff --git a/Engine/Plugins/Runtime/Steam/SteamVR/Source/SteamVRInputDevice/Private/SteamVRInputDevice.cpp b/Engine/Plugins/Runtime/Steam/SteamVR/Source/SteamVRInputDevice/Private/SteamVRInputDevice.cpp
--- a/Engine/Plugins/Runtime/Steam/SteamVR/Source/SteamVRInputDevice/Private/SteamVRInputDevice.cpp
+++ b/Engine/Plugins/Runtime/Steam/SteamVR/Source/SteamVRInputDevice/Private/SteamVRInputDevice.cpp
@@ -1673,8 +1673,15 @@ void FSteamVRInputDevice::GenerateActionBindings(TArray<FInputMapping> &InInputM
 				else
 				{
 					// Set cap sense input state
-					InputState.bIsCapSense = CurrentInputKeyName.Contains(TEXT("CapSense"), ESearchCase::CaseSensitive, ESearchDir::FromEnd) ||
-						CurrentInputKeyName.Contains(TEXT("_Touch"), ESearchCase::CaseSensitive, ESearchDir::FromEnd);
+					if (CurrentInputKeyName.Contains(TEXT("ValveIndex")) && InputState.bIsTrackpad)
+					{
+						InputState.bIsCapSense = true;					
+					}
+					else
+					{
+						InputState.bIsCapSense = CurrentInputKeyName.Contains(TEXT("CapSense"), ESearchCase::CaseSensitive, ESearchDir::FromEnd) ||
+							CurrentInputKeyName.Contains(TEXT("_Touch"), ESearchCase::CaseSensitive, ESearchDir::FromEnd);
+					}
 				}
 
 				// Check for DPad Keys
@@ -1737,10 +1744,9 @@ void FSteamVRInputDevice::GenerateActionBindings(TArray<FInputMapping> &InInputM
 				CacheMode = InputState.bIsTrigger || InputState.bIsGrip ? FName(TEXT("trigger")) : FName(TEXT("button"));
 				CacheMode = InputState.bIsPress ? FName(TEXT("button")) : CacheMode;
 				CacheMode = InputState.bIsTrackpad ? FName(TEXT("trackpad")) : CacheMode;
-				CacheMode = InputState.bIsJoystick ? FName(TEXT("joystick")) : CacheMode;
-				CacheMode = (InputState.bIsTrackpad || InputState.bIsJoystick) && !InputState.bIsAxis ? FName(TEXT("button")) : CacheMode;
+				CacheMode = InputState.bIsThumbstick ? FName(TEXT("joystick")) : CacheMode;
+				//CacheMode = (InputState.bIsTrackpad || InputState.bIsThumbstick) && !InputState.bIsAxis ? FName(TEXT("button")) : CacheMode;
 				CacheMode = InputState.bIsGrip ? FName(TEXT("button")) : CacheMode;
-				//CacheMode = InputState.bIsThumbstick ? FName(TEXT("button")) : CacheMode;
 				CacheMode = InputState.bIsPinchGrab || InputState.bIsGripGrab ? FName(TEXT("grab")) : CacheMode;
 
 				// Set Cache Path
@@ -1760,9 +1766,17 @@ void FSteamVRInputDevice::GenerateActionBindings(TArray<FInputMapping> &InInputM
 				{
 					CachePath = InputState.bIsLeft ? FString(TEXT(ACTION_PATH_TRACKPAD_LEFT)) : FString(TEXT(ACTION_PATH_TRACKPAD_RIGHT));
 				}
-				else if (InputState.bIsJoystick)
+				else if (InputState.bIsThumbstick)
 				{
-					CachePath = InputState.bIsLeft ? FString(TEXT(ACTION_PATH_JOYSTICK_LEFT)) : FString(TEXT(ACTION_PATH_JOYSTICK_RIGHT));
+					// Thumbstick vs Joystick (to conform with new UE naming scheme)
+					if (CurrentInputKeyName.Contains(TEXT("ValveIndex")))
+					{
+						CachePath = InputState.bIsLeft ? FString(TEXT(ACTION_PATH_THUMBSTICK_LEFT)) : FString(TEXT(ACTION_PATH_THUMBSTICK_RIGHT));
+					}
+					else
+					{
+						CachePath = InputState.bIsLeft ? FString(TEXT(ACTION_PATH_JOYSTICK_LEFT)) : FString(TEXT(ACTION_PATH_JOYSTICK_RIGHT));
+					}
 				}
 				else if (InputState.bIsGrip)
 				{
@@ -2009,7 +2023,7 @@ void FSteamVRInputDevice::GenerateActionBindings(TArray<FInputMapping> &InInputM
 						FString CurrentInputKeyName = SteamVRAxisKeyMapping.InputAxisKeyMapping.Key.ToString();
 						InputState.bIsTrigger = CurrentInputKeyName.Contains(TEXT("Trigger"), ESearchCase::CaseSensitive, ESearchDir::FromEnd);
 						InputState.bIsBumper = CurrentInputKeyName.Contains(TEXT("Bumper"), ESearchCase::CaseSensitive, ESearchDir::FromEnd);
-						InputState.bIsThumbstick = InputState.bIsJoystick = CurrentInputKeyName.Contains(TEXT("Thumbstick"), ESearchCase::CaseSensitive, ESearchDir::FromEnd);
+						InputState.bIsThumbstick = CurrentInputKeyName.Contains(TEXT("Thumbstick"), ESearchCase::CaseSensitive, ESearchDir::FromEnd);
 						InputState.bIsTrackpad = CurrentInputKeyName.Contains(TEXT("Trackpad"), ESearchCase::CaseSensitive, ESearchDir::FromEnd);
 						InputState.bIsGrip = CurrentInputKeyName.Contains(TEXT("Grip"), ESearchCase::CaseSensitive, ESearchDir::FromEnd);
 						InputState.bIsLeft = CurrentInputKeyName.Contains(TEXT("_Left_"), ESearchCase::CaseSensitive, ESearchDir::FromEnd);
@@ -2018,24 +2032,6 @@ void FSteamVRInputDevice::GenerateActionBindings(TArray<FInputMapping> &InInputM
 						InputState.bIsFaceButton2 = CurrentInputKeyName.Contains(TEXT("FaceButton2"), ESearchCase::CaseSensitive, ESearchDir::FromEnd) ||
 							CurrentInputKeyName.Contains(TEXT("_B_"));
 	
-						// Thumbstick vs Joystick (to conform with new UE naming scheme
-						if (CurrentInputKeyName.Contains(TEXT("ValveIndex")))
-						{
-							if (InputState.bIsThumbstick)
-							{
-								InputState.bIsThumbstick = true;
-								InputState.bIsJoystick = false;
-							}
-						}
-						else
-						{
-							if (InputState.bIsThumbstick)
-							{
-								InputState.bIsThumbstick = false;
-								InputState.bIsJoystick = true;
-							}
-						}
-
 						// Handle Oculus Touch
 						InputState.bIsXButton = InputState.bIsYButton = false;
 						if (CurrentInputKeyName.Contains(TEXT("OculusTouch")))
@@ -2086,7 +2082,6 @@ void FSteamVRInputDevice::GenerateActionBindings(TArray<FInputMapping> &InInputM
 						// Set Cache Mode
 						CacheMode = InputState.bIsTrigger || InputState.bIsGrip ? FName(TEXT("trigger")) : FName(TEXT("button"));
 						CacheMode = InputState.bIsTrackpad ? FName(TEXT("trackpad")) : CacheMode;
-						CacheMode = InputState.bIsJoystick ? FName(TEXT("joystick")) : CacheMode;
 						CacheMode = InputState.bIsGrip ? FName(TEXT("force_sensor")) : CacheMode;
 						CacheMode = InputState.bIsThumbstick ? FName(TEXT("joystick")) : CacheMode;
 						CacheMode = InputState.bIsPinchGrab || InputState.bIsGripGrab ? FName(TEXT("grab")) : CacheMode;
@@ -2109,16 +2104,20 @@ void FSteamVRInputDevice::GenerateActionBindings(TArray<FInputMapping> &InInputM
 						}
 						else if (InputState.bIsThumbstick)
 						{
-							CachePath = InputState.bIsLeft ? FString(TEXT(ACTION_PATH_THUMBSTICK_LEFT)) : FString(TEXT(ACTION_PATH_THUMBSTICK_RIGHT));
+							// Thumbstick vs Joystick (to conform with new UE naming scheme)
+							if (CurrentInputKeyName.Contains(TEXT("ValveIndex")))
+							{
+								CachePath = InputState.bIsLeft ? FString(TEXT(ACTION_PATH_THUMBSTICK_LEFT)) : FString(TEXT(ACTION_PATH_THUMBSTICK_RIGHT));
+							}
+							else
+							{
+								CachePath = InputState.bIsLeft ? FString(TEXT(ACTION_PATH_JOYSTICK_LEFT)) : FString(TEXT(ACTION_PATH_JOYSTICK_RIGHT));
+							}
 						}
 						else if (InputState.bIsTrackpad)
 						{
 							CachePath = InputState.bIsLeft ? FString(TEXT(ACTION_PATH_TRACKPAD_LEFT)) : FString(TEXT(ACTION_PATH_TRACKPAD_RIGHT));
 						}
-						else if (InputState.bIsJoystick)
-						{
-							CachePath = InputState.bIsLeft ? FString(TEXT(ACTION_PATH_JOYSTICK_LEFT)) : FString(TEXT(ACTION_PATH_JOYSTICK_RIGHT));
-						}
 						else if (InputState.bIsGrip)
 						{
 							CachePath = InputState.bIsLeft ? FString(TEXT(ACTION_PATH_GRIP_LEFT)) : FString(TEXT(ACTION_PATH_GRIP_RIGHT));
