Fix SteamVR depth submission. To get the Oculus-on-steam depth working correctly (i.e. having the dash intersect with scene content correctly) we applied a rescaling of the projection matrix. This works on an Oculus HMD for their time/spacewarp, but on a Steam device this breaks reprojection horribly.
This fix will make reprojection work, but the Oculus dash intersection will not work (this basically reverts to 4.23 and before behavior). Fixing this is, as far as we can tell, on the SteamVR runtime rather than on us, and an email thread has been started with Valve to cover this.

diff --git a/Engine/Plugins/Runtime/Steam/SteamVR/Source/SteamVR/Private/SteamVRRender.cpp b/Engine/Plugins/Runtime/Steam/SteamVR/Source/SteamVR/Private/SteamVRRender.cpp
--- a/Engine/Plugins/Runtime/Steam/SteamVR/Source/SteamVR/Private/SteamVRRender.cpp
+++ b/Engine/Plugins/Runtime/Steam/SteamVR/Source/SteamVR/Private/SteamVRRender.cpp
@@ -207,8 +207,6 @@ void FSteamVRHMD::D3D11Bridge::FinishRendering()
 		Texture.depth.vRange.v[1] = 0.0f;
 
 		Texture.depth.mProjection = ToHmdMatrix44(Plugin->GetStereoProjectionMatrix(eSSP_LEFT_EYE));
-		// Rescale the projection (our projection value is 10.0f here, since our units are cm, and SteamVR works in meters).
-		Texture.depth.mProjection.m[2][3] *= 0.01f;
 	}
 
 	vr::VRTextureBounds_t LeftBounds;
@@ -228,8 +226,6 @@ void FSteamVRHMD::D3D11Bridge::FinishRendering()
 	if (bSubmitDepth)
 	{
 		Texture.depth.mProjection = ToHmdMatrix44(Plugin->GetStereoProjectionMatrix(eSSP_RIGHT_EYE));
-		// Rescale the projection (our projection value is 10.0f here, since our units are cm, and SteamVR works in meters).
-		Texture.depth.mProjection.m[2][3] *= 0.01f;
 	}
 
 	Error = Plugin->VRCompositor->Submit(vr::Eye_Right, &Texture, &RightBounds, Flags);
