Fix MSAA broken on SteamVR (AllocateDepthTargets should always just use "1" in the same fashion as other plugins, and also the color target).

diff --git a/Engine/Plugins/Runtime/Steam/SteamVR/Source/SteamVR/Private/SteamVRHMD.cpp b/Engine/Plugins/Runtime/Steam/SteamVR/Source/SteamVR/Private/SteamVRHMD.cpp
--- a/Engine/Plugins/Runtime/Steam/SteamVR/Source/SteamVR/Private/SteamVRHMD.cpp
+++ b/Engine/Plugins/Runtime/Steam/SteamVR/Source/SteamVR/Private/SteamVRHMD.cpp
@@ -1653,7 +1653,7 @@ bool FSteamVRHMD::AllocateRenderTargetTexture(uint32 Index, uint32 SizeX, uint32
 }
 
 bool FSteamVRHMD::AllocateDepthTexture(uint32 Index, uint32 SizeX, uint32 SizeY, uint8 Format, uint32 NumMips, uint32 Flags, uint32 TargetableTextureFlags, 
-	FTexture2DRHIRef& OutTargetableTexture, FTexture2DRHIRef& OutShaderResourceTexture, uint32 NumSamples /* = 1 */)
+	FTexture2DRHIRef& OutTargetableTexture, FTexture2DRHIRef& OutShaderResourceTexture, uint32 NumSamples /* ignored, we always use 1 */)
 {
 	if (!IsStereoEnabled() || pBridge == nullptr)
 	{
@@ -1699,7 +1699,7 @@ bool FSteamVRHMD::AllocateDepthTexture(uint32 Index, uint32 SizeX, uint32 SizeY,
 			CreateInfo,
 			TargetableTexture,
 			ShaderResourceTexture,
-			NumSamples);
+			1);
 
 		check(TargetableTexture == ShaderResourceTexture);
 
@@ -1709,7 +1709,7 @@ bool FSteamVRHMD::AllocateDepthTexture(uint32 Index, uint32 SizeX, uint32 SizeY,
 		if (BindingTexture == nullptr)
 		{
 			FTexture2DRHIRef TempTargetableTexture, TempShaderResourceTexture;
-			RHICreateTargetableShaderResource2D(SizeX, SizeY, PF_DepthStencil, 1, Flags, TargetableTextureFlags, false, CreateInfo, TempTargetableTexture, TempShaderResourceTexture, NumSamples);
+			RHICreateTargetableShaderResource2D(SizeX, SizeY, PF_DepthStencil, 1, Flags, TargetableTextureFlags, false, CreateInfo, TempTargetableTexture, TempShaderResourceTexture, 1);
 			check(TempTargetableTexture == TempShaderResourceTexture);
 			BindingTexture = TempTargetableTexture;
 
