Automatically add actions for the VREditor when using the SteamVR plugin

diff --git a/Engine/Plugins/Runtime/Steam/SteamVR/Source/SteamVRInputDevice/Private/SteamVRInputDevice.cpp b/Engine/Plugins/Runtime/Steam/SteamVR/Source/SteamVRInputDevice/Private/SteamVRInputDevice.cpp
--- a/Engine/Plugins/Runtime/Steam/SteamVR/Source/SteamVRInputDevice/Private/SteamVRInputDevice.cpp
+++ b/Engine/Plugins/Runtime/Steam/SteamVR/Source/SteamVRInputDevice/Private/SteamVRInputDevice.cpp
@@ -1241,6 +1241,84 @@ void FSteamVRInputDevice::OnActionMappingsChanged()
 	this->GenerateActionManifest(true, true, true, true);
 }
 
+void FSteamVRInputDevice::OnVREditingModeEnter()
+{
+	// Set Input Settings
+	auto InputSettings = GetMutableDefault<UInputSettings>();
+
+	InputSettings->AddActionMapping(FInputActionKeyMapping("VREditor_WorldMovementL", EKeys::MotionController_Left_Grip1), false);
+	InputSettings->AddActionMapping(FInputActionKeyMapping("VREditor_TouchL", EKeys::Vive_Left_Trackpad_Touch), false);
+	InputSettings->AddAxisMapping(FInputAxisKeyMapping("VREditor_TriggerAxisL", EKeys::MotionController_Left_TriggerAxis), false);
+	InputSettings->AddAxisMapping(FInputAxisKeyMapping("VREditor_TrackpadPositionXL", EKeys::MotionController_Left_Thumbstick_X), false);
+	InputSettings->AddAxisMapping(FInputAxisKeyMapping("VREditor_TrackpadPositionYL", EKeys::MotionController_Left_Thumbstick_Y), false);
+	InputSettings->AddActionMapping(FInputActionKeyMapping("VREditor_ConfirmRadialSelectionL", EKeys::MotionController_Left_Thumbstick), false);
+
+	InputSettings->AddActionMapping(FInputActionKeyMapping("VREditor_ModifierL", EKeys::Vive_Left_Menu_Click), false);
+
+	InputSettings->AddActionMapping(FInputActionKeyMapping("VREditor_TrackpadDownL", EKeys::Vive_Left_Trackpad_Down), false); // down
+	InputSettings->AddActionMapping(FInputActionKeyMapping("VREditor_TrackpadLeftL", EKeys::Vive_Left_Trackpad_Left), false);
+	InputSettings->AddActionMapping(FInputActionKeyMapping("VREditor_TrackpadRightL", EKeys::Vive_Left_Trackpad_Right), false);
+	InputSettings->AddActionMapping(FInputActionKeyMapping("VREditor_TrackpadUpL", EKeys::Vive_Left_Trackpad_Up), false);
+
+	InputSettings->AddActionMapping(FInputActionKeyMapping("VREditor_WorldMovementR", EKeys::MotionController_Right_Grip1), false);
+	InputSettings->AddActionMapping(FInputActionKeyMapping("VREditor_TouchR", EKeys::Vive_Right_Trackpad_Touch), false);
+	InputSettings->AddAxisMapping(FInputAxisKeyMapping("VREditor_TriggerAxisR", EKeys::MotionController_Right_TriggerAxis), false);
+	InputSettings->AddAxisMapping(FInputAxisKeyMapping("VREditor_TrackpadPositionXR", EKeys::MotionController_Right_Thumbstick_X), false);
+	InputSettings->AddAxisMapping(FInputAxisKeyMapping("VREditor_TrackpadPositionYR", EKeys::MotionController_Right_Thumbstick_Y), false);
+	InputSettings->AddActionMapping(FInputActionKeyMapping("VREditor_ConfirmRadialSelectionR", EKeys::MotionController_Right_Thumbstick), false);
+
+
+	InputSettings->AddActionMapping(FInputActionKeyMapping("VREditor_ModifierR", EKeys::Vive_Right_Menu_Click), false);
+
+	InputSettings->AddActionMapping(FInputActionKeyMapping("VREditor_TrackpadDownR", EKeys::Vive_Right_Trackpad_Down), false); // down
+	InputSettings->AddActionMapping(FInputActionKeyMapping("VREditor_TrackpadLeftR", EKeys::Vive_Right_Trackpad_Left), false);
+	InputSettings->AddActionMapping(FInputActionKeyMapping("VREditor_TrackpadRightR", EKeys::Vive_Right_Trackpad_Right), false);
+	InputSettings->AddActionMapping(FInputActionKeyMapping("VREditor_TrackpadUpR", EKeys::Vive_Right_Trackpad_Up), false);
+
+	InputSettings->ForceRebuildKeymaps();
+
+	this->GenerateActionManifest(true, true, true, true);
+}
+
+void FSteamVRInputDevice::OnVREditingModeExit()
+{
+	// Set Input Settings
+	auto InputSettings = GetMutableDefault<UInputSettings>();
+
+	InputSettings->RemoveActionMapping(FInputActionKeyMapping("VREditor_WorldMovementL", EKeys::MotionController_Left_Grip1), false);
+	InputSettings->RemoveActionMapping(FInputActionKeyMapping("VREditor_TouchL", EKeys::Vive_Left_Trackpad_Touch), false);
+	InputSettings->RemoveAxisMapping(FInputAxisKeyMapping("VREditor_TriggerAxisL", EKeys::MotionController_Left_TriggerAxis), false);
+	InputSettings->RemoveAxisMapping(FInputAxisKeyMapping("VREditor_TrackpadPositionXL", EKeys::MotionController_Left_Thumbstick_X), false);
+	InputSettings->RemoveAxisMapping(FInputAxisKeyMapping("VREditor_TrackpadPositionYL", EKeys::MotionController_Left_Thumbstick_Y), false);
+	InputSettings->RemoveActionMapping(FInputActionKeyMapping("VREditor_ConfirmRadialSelectionL", EKeys::MotionController_Left_Thumbstick), false);
+
+	InputSettings->RemoveActionMapping(FInputActionKeyMapping("VREditor_ModifierL", EKeys::Vive_Left_Menu_Click), false);
+
+	InputSettings->RemoveActionMapping(FInputActionKeyMapping("VREditor_TrackpadDownL", EKeys::Vive_Left_Trackpad_Down), false); // down
+	InputSettings->RemoveActionMapping(FInputActionKeyMapping("VREditor_TrackpadLeftL", EKeys::Vive_Left_Trackpad_Left), false);
+	InputSettings->RemoveActionMapping(FInputActionKeyMapping("VREditor_TrackpadRightL", EKeys::Vive_Left_Trackpad_Right), false);
+	InputSettings->RemoveActionMapping(FInputActionKeyMapping("VREditor_TrackpadUpL", EKeys::Vive_Left_Trackpad_Up), false);
+
+	InputSettings->RemoveActionMapping(FInputActionKeyMapping("VREditor_WorldMovementR", EKeys::MotionController_Right_Grip1), false);
+	InputSettings->RemoveActionMapping(FInputActionKeyMapping("VREditor_TouchR", EKeys::Vive_Right_Trackpad_Touch), false);
+	InputSettings->RemoveAxisMapping(FInputAxisKeyMapping("VREditor_TriggerAxisR", EKeys::MotionController_Right_TriggerAxis), false);
+	InputSettings->RemoveAxisMapping(FInputAxisKeyMapping("VREditor_TrackpadPositionXR", EKeys::MotionController_Right_Thumbstick_X), false);
+	InputSettings->RemoveAxisMapping(FInputAxisKeyMapping("VREditor_TrackpadPositionYR", EKeys::MotionController_Right_Thumbstick_Y), false);
+	InputSettings->RemoveActionMapping(FInputActionKeyMapping("VREditor_ConfirmRadialSelectionR", EKeys::MotionController_Right_Thumbstick), false);
+
+
+	InputSettings->RemoveActionMapping(FInputActionKeyMapping("VREditor_ModifierR", EKeys::Vive_Right_Menu_Click), false);
+
+	InputSettings->RemoveActionMapping(FInputActionKeyMapping("VREditor_TrackpadDownR", EKeys::Vive_Right_Trackpad_Down), false); // down
+	InputSettings->RemoveActionMapping(FInputActionKeyMapping("VREditor_TrackpadLeftR", EKeys::Vive_Right_Trackpad_Left), false);
+	InputSettings->RemoveActionMapping(FInputActionKeyMapping("VREditor_TrackpadRightR", EKeys::Vive_Right_Trackpad_Right), false);
+	InputSettings->RemoveActionMapping(FInputActionKeyMapping("VREditor_TrackpadUpR", EKeys::Vive_Right_Trackpad_Up), false);
+
+	InputSettings->ForceRebuildKeymaps();
+
+	this->GenerateActionManifest(true, true, true, true);
+}
+
 bool FSteamVRInputDevice::GenerateAppManifest(FString ManifestPath, FString ProjectName, FString& OutAppKey, FString& OutAppManifestPath)
 {
 	// Set SteamVR AppKey
@@ -1773,6 +1851,10 @@ void FSteamVRInputDevice::GenerateActionBindings(TArray<FInputMapping> &InInputM
 					{
 						CachePath = InputState.bIsLeft ? FString(TEXT(ACTION_PATH_THUMBSTICK_LEFT)) : FString(TEXT(ACTION_PATH_THUMBSTICK_RIGHT));
 					}
+					else if (Controller.Description.Contains(TEXT("Vive")))
+					{
+						CachePath = InputState.bIsLeft ? FString(TEXT(ACTION_PATH_TRACKPAD_LEFT)) : FString(TEXT(ACTION_PATH_TRACKPAD_RIGHT));
+					}
 					else
 					{
 						CachePath = InputState.bIsLeft ? FString(TEXT(ACTION_PATH_JOYSTICK_LEFT)) : FString(TEXT(ACTION_PATH_JOYSTICK_RIGHT));
@@ -2109,6 +2191,10 @@ void FSteamVRInputDevice::GenerateActionBindings(TArray<FInputMapping> &InInputM
 							{
 								CachePath = InputState.bIsLeft ? FString(TEXT(ACTION_PATH_THUMBSTICK_LEFT)) : FString(TEXT(ACTION_PATH_THUMBSTICK_RIGHT));
 							}
+							else if (Controller.Description.Contains(TEXT("Vive")))
+							{
+								CachePath = InputState.bIsLeft ? FString(TEXT(ACTION_PATH_TRACKPAD_LEFT)) : FString(TEXT(ACTION_PATH_TRACKPAD_RIGHT));
+							}
 							else
 							{
 								CachePath = InputState.bIsLeft ? FString(TEXT(ACTION_PATH_JOYSTICK_LEFT)) : FString(TEXT(ACTION_PATH_JOYSTICK_RIGHT));
diff --git a/Engine/Plugins/Runtime/Steam/SteamVR/Source/SteamVRInputDevice/Private/SteamVRInputDeviceModule.cpp b/Engine/Plugins/Runtime/Steam/SteamVR/Source/SteamVRInputDevice/Private/SteamVRInputDeviceModule.cpp
--- a/Engine/Plugins/Runtime/Steam/SteamVR/Source/SteamVRInputDevice/Private/SteamVRInputDeviceModule.cpp
+++ b/Engine/Plugins/Runtime/Steam/SteamVR/Source/SteamVRInputDevice/Private/SteamVRInputDeviceModule.cpp
@@ -38,6 +38,7 @@ POSSIBILITY OF SUCH DAMAGE.
 #if WITH_EDITOR
 #include "Editor/EditorEngine.h"
 #include "Editor.h"
+#include "IVREditorModule.h"
 #endif
 
 class FSteamVRInputDeviceModule : public ISteamVRInputDeviceModule
@@ -51,6 +52,12 @@ class FSteamVRInputDeviceModule : public ISteamVRInputDeviceModule
 			TSharedPtr<class FSteamVRInputDevice> Device(new FSteamVRInputDevice(InMessageHandler));
 #if WITH_EDITOR
 			FEditorDelegates::OnActionAxisMappingsChanged.AddSP(Device.ToSharedRef(), &FSteamVRInputDevice::OnActionMappingsChanged);
+
+			if (IVREditorModule::IsAvailable())
+			{
+				IVREditorModule::Get().OnVREditingModeEnter().AddSP(Device.ToSharedRef(), &FSteamVRInputDevice::OnVREditingModeEnter);
+				IVREditorModule::Get().OnVREditingModeExit().AddSP(Device.ToSharedRef(), &FSteamVRInputDevice::OnVREditingModeExit);
+			}
 #endif
 			return Device;
 		}
diff --git a/Engine/Plugins/Runtime/Steam/SteamVR/Source/SteamVRInputDevice/Public/SteamVRInputDevice.h b/Engine/Plugins/Runtime/Steam/SteamVR/Source/SteamVRInputDevice/Public/SteamVRInputDevice.h
--- a/Engine/Plugins/Runtime/Steam/SteamVR/Source/SteamVRInputDevice/Public/SteamVRInputDevice.h
+++ b/Engine/Plugins/Runtime/Steam/SteamVR/Source/SteamVRInputDevice/Public/SteamVRInputDevice.h
@@ -129,6 +129,12 @@ public:
 	/** Have the manifest and bindings bindings regenerated. Used by a UI callback */
 	void OnActionMappingsChanged();
 
+	/** Add VR editor actions and bindings. Used by a UI callback */
+	void OnVREditingModeEnter();
+
+	/** Remove VR editor actions and bindings. Used by a UI callback */
+	void OnVREditingModeExit();
+
 	/** Reload the action manifest to the SteamVR system. Use if you changed the action manifest manually */
 	void ReloadActionManifest();
 #endif
diff --git a/Engine/Source/Editor/VREditor/VREditorInteractor.cpp b/Engine/Source/Editor/VREditor/VREditorInteractor.cpp
--- a/Engine/Source/Editor/VREditor/VREditorInteractor.cpp
+++ b/Engine/Source/Editor/VREditor/VREditorInteractor.cpp
@@ -165,7 +165,7 @@ void UVREditorInteractor::Init_Implementation( UVREditorMode* InVRMode )
 		AddKeyAction( EKeys::MotionController_Left_Grip1, FViewportActionKeyInput( ViewportWorldActionTypes::WorldMovement ) );
 		AddKeyAction( UVREditorInteractor::MotionController_Left_FullyPressedTriggerAxis, FViewportActionKeyInput( ViewportWorldActionTypes::SelectAndMove_FullyPressed ) );
 		AddKeyAction( UVREditorInteractor::MotionController_Left_PressedTriggerAxis, FViewportActionKeyInput( ViewportWorldActionTypes::SelectAndMove ) );
-		AddKeyAction( SteamVRControllerKeyNames::Touch0, FViewportActionKeyInput( VRActionTypes::Touch ) );
+		AddKeyAction( EKeys::Vive_Left_Trackpad_Touch, FViewportActionKeyInput( VRActionTypes::Touch ) );
 		AddKeyAction( EKeys::MotionController_Left_TriggerAxis, FViewportActionKeyInput( UVREditorInteractor::TriggerAxis ) );
 		AddKeyAction( EKeys::MotionController_Left_Thumbstick_X, FViewportActionKeyInput( UVREditorInteractor::TrackpadPositionX ) );
 		AddKeyAction( EKeys::MotionController_Left_Thumbstick_Y, FViewportActionKeyInput( UVREditorInteractor::TrackpadPositionY ) );
@@ -173,12 +173,12 @@ void UVREditorInteractor::Init_Implementation( UVREditorMode* InVRMode )
 
 		if (HMDDeviceType == SteamVRDeviceType)
 		{
-			AddKeyAction( EKeys::MotionController_Left_Shoulder, FViewportActionKeyInput( VRActionTypes::Modifier ) );
+			AddKeyAction( EKeys::Vive_Left_Menu_Click, FViewportActionKeyInput( VRActionTypes::Modifier ) );
 
-			AddKeyAction( EKeys::MotionController_Left_FaceButton3, FViewportActionKeyInput( VRActionTypes::TrackpadDown ) ); // down
-			AddKeyAction( EKeys::MotionController_Left_FaceButton4, FViewportActionKeyInput( VRActionTypes::TrackpadLeft ) );
-			AddKeyAction( EKeys::MotionController_Left_FaceButton2, FViewportActionKeyInput( VRActionTypes::TrackpadRight ) );
-			AddKeyAction( EKeys::MotionController_Left_FaceButton1, FViewportActionKeyInput( VRActionTypes::TrackpadUp ) );
+			AddKeyAction( EKeys::Vive_Left_Trackpad_Down, FViewportActionKeyInput( VRActionTypes::TrackpadDown ) ); // down
+			AddKeyAction( EKeys::Vive_Left_Trackpad_Left, FViewportActionKeyInput( VRActionTypes::TrackpadLeft ) );
+			AddKeyAction( EKeys::Vive_Left_Trackpad_Right, FViewportActionKeyInput( VRActionTypes::TrackpadRight ) );
+			AddKeyAction( EKeys::Vive_Left_Trackpad_Up, FViewportActionKeyInput( VRActionTypes::TrackpadUp ) );
 		}
 		else if (HMDDeviceType == OculusDeviceType)
 		{
@@ -197,7 +197,7 @@ void UVREditorInteractor::Init_Implementation( UVREditorMode* InVRMode )
 		AddKeyAction( EKeys::MotionController_Right_Grip1, FViewportActionKeyInput( ViewportWorldActionTypes::WorldMovement ) );
 		AddKeyAction( UVREditorInteractor::MotionController_Right_FullyPressedTriggerAxis, FViewportActionKeyInput( ViewportWorldActionTypes::SelectAndMove_FullyPressed ) );
 		AddKeyAction( UVREditorInteractor::MotionController_Right_PressedTriggerAxis, FViewportActionKeyInput( ViewportWorldActionTypes::SelectAndMove ) );
-		AddKeyAction( SteamVRControllerKeyNames::Touch1, FViewportActionKeyInput( VRActionTypes::Touch ) );
+		AddKeyAction( EKeys::Vive_Right_Trackpad_Touch, FViewportActionKeyInput( VRActionTypes::Touch ) );
 		AddKeyAction( EKeys::MotionController_Right_TriggerAxis, FViewportActionKeyInput( UVREditorInteractor::TriggerAxis ) );
 		AddKeyAction( EKeys::MotionController_Right_Thumbstick_X, FViewportActionKeyInput( UVREditorInteractor::TrackpadPositionX ) );
 		AddKeyAction( EKeys::MotionController_Right_Thumbstick_Y, FViewportActionKeyInput( UVREditorInteractor::TrackpadPositionY ) );
@@ -205,12 +205,12 @@ void UVREditorInteractor::Init_Implementation( UVREditorMode* InVRMode )
 
 		if (HMDDeviceType == SteamVRDeviceType)
 		{
-			AddKeyAction( EKeys::MotionController_Right_Shoulder, FViewportActionKeyInput( VRActionTypes::Modifier ) );
+			AddKeyAction( EKeys::Vive_Right_Menu_Click, FViewportActionKeyInput( VRActionTypes::Modifier ) );
 
-			AddKeyAction( EKeys::MotionController_Right_FaceButton3, FViewportActionKeyInput( VRActionTypes::TrackpadDown ) ); // down
-			AddKeyAction( EKeys::MotionController_Right_FaceButton4, FViewportActionKeyInput( VRActionTypes::TrackpadLeft ) );
-			AddKeyAction( EKeys::MotionController_Right_FaceButton2, FViewportActionKeyInput( VRActionTypes::TrackpadRight ) );
-			AddKeyAction( EKeys::MotionController_Right_FaceButton1, FViewportActionKeyInput( VRActionTypes::TrackpadUp ) );
+			AddKeyAction( EKeys::Vive_Right_Trackpad_Down, FViewportActionKeyInput( VRActionTypes::TrackpadDown ) ); // down
+			AddKeyAction( EKeys::Vive_Right_Trackpad_Left, FViewportActionKeyInput( VRActionTypes::TrackpadLeft ) );
+			AddKeyAction( EKeys::Vive_Right_Trackpad_Right, FViewportActionKeyInput( VRActionTypes::TrackpadRight ) );
+			AddKeyAction( EKeys::Vive_Right_Trackpad_Up, FViewportActionKeyInput( VRActionTypes::TrackpadUp ) );
 		}
 		else if (HMDDeviceType == OculusDeviceType)
 		{
