Disable depth submission on SteamVR by default, re-word CVar comment.

diff --git a/Engine/Plugins/Runtime/Steam/SteamVR/Source/SteamVR/Private/SteamVRRender.cpp b/Engine/Plugins/Runtime/Steam/SteamVR/Source/SteamVR/Private/SteamVRRender.cpp
--- a/Engine/Plugins/Runtime/Steam/SteamVR/Source/SteamVR/Private/SteamVRRender.cpp
+++ b/Engine/Plugins/Runtime/Steam/SteamVR/Source/SteamVR/Private/SteamVRRender.cpp
@@ -31,8 +31,8 @@ static TAutoConsoleVariable<int32> CUsePostPresentHandoff(TEXT("vr.SteamVR.UsePo
 
 static TAutoConsoleVariable<int> CVarEnableDepthSubmission(
 	TEXT("vr.EnableSteamVRDepthSubmission"),
-	1,
-	TEXT("By default, depth is passed through in SteamVR for devices that support depth. Set this flag to 0 to disable depth submission."),
+	0,
+	TEXT("By default, depth is not passed through in SteamVR for devices that support depth. Set this flag to 1 to enable depth submission, 0 to disable."),
 	ECVF_Default);
 
 void FSteamVRHMD::DrawDistortionMesh_RenderThread(struct FRenderingCompositePassContext& Context, const FIntPoint& TextureSize)
