Fixed radial menu not working in VR Editor/Scouting

diff --git a/Engine/Plugins/Runtime/Steam/SteamVR/Source/SteamVRInputDevice/Private/SteamVRInputDevice.cpp b/Engine/Plugins/Runtime/Steam/SteamVR/Source/SteamVRInputDevice/Private/SteamVRInputDevice.cpp
--- a/Engine/Plugins/Runtime/Steam/SteamVR/Source/SteamVRInputDevice/Private/SteamVRInputDevice.cpp
+++ b/Engine/Plugins/Runtime/Steam/SteamVR/Source/SteamVRInputDevice/Private/SteamVRInputDevice.cpp
@@ -1249,8 +1249,8 @@ void FSteamVRInputDevice::OnVREditingModeEnter()
 	InputSettings->AddActionMapping(FInputActionKeyMapping("VREditor_WorldMovementL", EKeys::MotionController_Left_Grip1), false);
 	InputSettings->AddActionMapping(FInputActionKeyMapping("VREditor_TouchL", EKeys::Vive_Left_Trackpad_Touch), false);
 	InputSettings->AddAxisMapping(FInputAxisKeyMapping("VREditor_TriggerAxisL", EKeys::MotionController_Left_TriggerAxis), false);
-	InputSettings->AddAxisMapping(FInputAxisKeyMapping("VREditor_TrackpadPositionXL", EKeys::MotionController_Left_Thumbstick_X), false);
-	InputSettings->AddAxisMapping(FInputAxisKeyMapping("VREditor_TrackpadPositionYL", EKeys::MotionController_Left_Thumbstick_Y), false);
+	InputSettings->AddAxisMapping(FInputAxisKeyMapping("VREditor_TrackpadPositionL_X", EKeys::MotionController_Left_Thumbstick_X), false);
+	InputSettings->AddAxisMapping(FInputAxisKeyMapping("VREditor_TrackpadPositionL_Y", EKeys::MotionController_Left_Thumbstick_Y), false);
 	InputSettings->AddActionMapping(FInputActionKeyMapping("VREditor_ConfirmRadialSelectionL", EKeys::MotionController_Left_Thumbstick), false);
 
 	InputSettings->AddActionMapping(FInputActionKeyMapping("VREditor_ModifierL", EKeys::Vive_Left_Menu_Click), false);
@@ -1263,8 +1263,8 @@ void FSteamVRInputDevice::OnVREditingModeEnter()
 	InputSettings->AddActionMapping(FInputActionKeyMapping("VREditor_WorldMovementR", EKeys::MotionController_Right_Grip1), false);
 	InputSettings->AddActionMapping(FInputActionKeyMapping("VREditor_TouchR", EKeys::Vive_Right_Trackpad_Touch), false);
 	InputSettings->AddAxisMapping(FInputAxisKeyMapping("VREditor_TriggerAxisR", EKeys::MotionController_Right_TriggerAxis), false);
-	InputSettings->AddAxisMapping(FInputAxisKeyMapping("VREditor_TrackpadPositionXR", EKeys::MotionController_Right_Thumbstick_X), false);
-	InputSettings->AddAxisMapping(FInputAxisKeyMapping("VREditor_TrackpadPositionYR", EKeys::MotionController_Right_Thumbstick_Y), false);
+	InputSettings->AddAxisMapping(FInputAxisKeyMapping("VREditor_TrackpadPositionR_X", EKeys::MotionController_Right_Thumbstick_X), false);
+	InputSettings->AddAxisMapping(FInputAxisKeyMapping("VREditor_TrackpadPositionR_Y", EKeys::MotionController_Right_Thumbstick_Y), false);
 	InputSettings->AddActionMapping(FInputActionKeyMapping("VREditor_ConfirmRadialSelectionR", EKeys::MotionController_Right_Thumbstick), false);
 
 
@@ -1288,8 +1288,8 @@ void FSteamVRInputDevice::OnVREditingModeExit()
 	InputSettings->RemoveActionMapping(FInputActionKeyMapping("VREditor_WorldMovementL", EKeys::MotionController_Left_Grip1), false);
 	InputSettings->RemoveActionMapping(FInputActionKeyMapping("VREditor_TouchL", EKeys::Vive_Left_Trackpad_Touch), false);
 	InputSettings->RemoveAxisMapping(FInputAxisKeyMapping("VREditor_TriggerAxisL", EKeys::MotionController_Left_TriggerAxis), false);
-	InputSettings->RemoveAxisMapping(FInputAxisKeyMapping("VREditor_TrackpadPositionXL", EKeys::MotionController_Left_Thumbstick_X), false);
-	InputSettings->RemoveAxisMapping(FInputAxisKeyMapping("VREditor_TrackpadPositionYL", EKeys::MotionController_Left_Thumbstick_Y), false);
+	InputSettings->RemoveAxisMapping(FInputAxisKeyMapping("VREditor_TrackpadPositionL_X", EKeys::MotionController_Left_Thumbstick_X), false);
+	InputSettings->RemoveAxisMapping(FInputAxisKeyMapping("VREditor_TrackpadPositionL_Y", EKeys::MotionController_Left_Thumbstick_Y), false);
 	InputSettings->RemoveActionMapping(FInputActionKeyMapping("VREditor_ConfirmRadialSelectionL", EKeys::MotionController_Left_Thumbstick), false);
 
 	InputSettings->RemoveActionMapping(FInputActionKeyMapping("VREditor_ModifierL", EKeys::Vive_Left_Menu_Click), false);
@@ -1302,8 +1302,8 @@ void FSteamVRInputDevice::OnVREditingModeExit()
 	InputSettings->RemoveActionMapping(FInputActionKeyMapping("VREditor_WorldMovementR", EKeys::MotionController_Right_Grip1), false);
 	InputSettings->RemoveActionMapping(FInputActionKeyMapping("VREditor_TouchR", EKeys::Vive_Right_Trackpad_Touch), false);
 	InputSettings->RemoveAxisMapping(FInputAxisKeyMapping("VREditor_TriggerAxisR", EKeys::MotionController_Right_TriggerAxis), false);
-	InputSettings->RemoveAxisMapping(FInputAxisKeyMapping("VREditor_TrackpadPositionXR", EKeys::MotionController_Right_Thumbstick_X), false);
-	InputSettings->RemoveAxisMapping(FInputAxisKeyMapping("VREditor_TrackpadPositionYR", EKeys::MotionController_Right_Thumbstick_Y), false);
+	InputSettings->RemoveAxisMapping(FInputAxisKeyMapping("VREditor_TrackpadPositionR_X", EKeys::MotionController_Right_Thumbstick_X), false);
+	InputSettings->RemoveAxisMapping(FInputAxisKeyMapping("VREditor_TrackpadPositionR_Y", EKeys::MotionController_Right_Thumbstick_Y), false);
 	InputSettings->RemoveActionMapping(FInputActionKeyMapping("VREditor_ConfirmRadialSelectionR", EKeys::MotionController_Right_Thumbstick), false);
 
 
@@ -1818,6 +1818,13 @@ void FSteamVRInputDevice::GenerateActionBindings(TArray<FInputMapping> &InInputM
 					InputState.bIsGripGrab = false;
 				}
 
+				// Handle Vive controllers not having a thumbstick
+				if (InputState.bIsThumbstick && Controller.Description.Contains(TEXT("Vive")))
+				{
+					InputState.bIsTrackpad = true;
+					InputState.bIsThumbstick = false;
+				}
+
 				// Set Cache Mode
 				CacheMode = InputState.bIsTrigger || InputState.bIsGrip ? FName(TEXT("trigger")) : FName(TEXT("button"));
 				CacheMode = InputState.bIsPress ? FName(TEXT("button")) : CacheMode;
@@ -1851,10 +1858,6 @@ void FSteamVRInputDevice::GenerateActionBindings(TArray<FInputMapping> &InInputM
 					{
 						CachePath = InputState.bIsLeft ? FString(TEXT(ACTION_PATH_THUMBSTICK_LEFT)) : FString(TEXT(ACTION_PATH_THUMBSTICK_RIGHT));
 					}
-					else if (Controller.Description.Contains(TEXT("Vive")))
-					{
-						CachePath = InputState.bIsLeft ? FString(TEXT(ACTION_PATH_TRACKPAD_LEFT)) : FString(TEXT(ACTION_PATH_TRACKPAD_RIGHT));
-					}
 					else
 					{
 						CachePath = InputState.bIsLeft ? FString(TEXT(ACTION_PATH_JOYSTICK_LEFT)) : FString(TEXT(ACTION_PATH_JOYSTICK_RIGHT));
@@ -2160,6 +2163,13 @@ void FSteamVRInputDevice::GenerateActionBindings(TArray<FInputMapping> &InInputM
 							InputState.bIsPinchGrab = false;
 							InputState.bIsGripGrab = false;
 						}
+
+						// Handle Vive controllers not having a thumbstick
+						if (InputState.bIsThumbstick && Controller.Description.Contains(TEXT("Vive")))
+						{
+							InputState.bIsTrackpad = true;
+							InputState.bIsThumbstick = false;
+						}
 	
 						// Set Cache Mode
 						CacheMode = InputState.bIsTrigger || InputState.bIsGrip ? FName(TEXT("trigger")) : FName(TEXT("button"));
@@ -2191,10 +2201,6 @@ void FSteamVRInputDevice::GenerateActionBindings(TArray<FInputMapping> &InInputM
 							{
 								CachePath = InputState.bIsLeft ? FString(TEXT(ACTION_PATH_THUMBSTICK_LEFT)) : FString(TEXT(ACTION_PATH_THUMBSTICK_RIGHT));
 							}
-							else if (Controller.Description.Contains(TEXT("Vive")))
-							{
-								CachePath = InputState.bIsLeft ? FString(TEXT(ACTION_PATH_TRACKPAD_LEFT)) : FString(TEXT(ACTION_PATH_TRACKPAD_RIGHT));
-							}
 							else
 							{
 								CachePath = InputState.bIsLeft ? FString(TEXT(ACTION_PATH_JOYSTICK_LEFT)) : FString(TEXT(ACTION_PATH_JOYSTICK_RIGHT));
@@ -3631,8 +3637,16 @@ void FSteamVRInputDevice::ProcessActionEvents(FSteamVRInputActionSet SteamVRInpu
 				{
 					// Test what we're receiving from SteamVR
 					//UE_LOG(LogTemp, Warning, TEXT("Handle %s KeyY %s Y-Value {%f}"), *Action.Path, *FoundKeyY.GetFName().ToString(), AnalogData.y);
-
-					Action.Value.Y = AnalogData.y;
+					
+					FString KeyString = Action.KeyY.ToString();
+					if (KeyString.Contains(TEXT("MotionController")) && KeyString.Contains(TEXT("_Y")))
+					{
+						Action.Value.Y = -AnalogData.y;
+					}
+					else
+					{
+						Action.Value.Y = AnalogData.y;
+					}
 					MessageHandler->OnControllerAnalog(Action.KeyY, 0, Action.Value.Y);
 				}
 			}
