Remove invalid assert when re-creating render chain for XR when switching between VRMode and VR PIE.

diff --git a/Engine/Plugins/Runtime/Steam/SteamVR/Source/SteamVR/Private/SteamVRRender.cpp b/Engine/Plugins/Runtime/Steam/SteamVR/Source/SteamVR/Private/SteamVRRender.cpp
--- a/Engine/Plugins/Runtime/Steam/SteamVR/Source/SteamVR/Private/SteamVRRender.cpp
+++ b/Engine/Plugins/Runtime/Steam/SteamVR/Source/SteamVR/Private/SteamVRRender.cpp
@@ -249,7 +249,6 @@ void FSteamVRHMD::D3D11Bridge::UpdateViewport(const FViewport& Viewport, FRHIVie
 
 	const FTexture2DRHIRef& RT = Viewport.GetRenderTargetTexture();
 	check(IsValidRef(RT));
-	check(RT->GetTexture2D() == SwapChain->GetTexture2D());
 }
 
 FSteamVRHMD::D3D12Bridge::D3D12Bridge(FSteamVRHMD* plugin)
