Handle cancelling the SteamVR binding overwrite warning as "no"

diff --git a/Engine/Plugins/Runtime/Steam/SteamVR/Source/SteamVRInputDevice/Private/SteamVRInputDevice.cpp b/Engine/Plugins/Runtime/Steam/SteamVR/Source/SteamVRInputDevice/Private/SteamVRInputDevice.cpp
--- a/Engine/Plugins/Runtime/Steam/SteamVR/Source/SteamVRInputDevice/Private/SteamVRInputDevice.cpp
+++ b/Engine/Plugins/Runtime/Steam/SteamVR/Source/SteamVRInputDevice/Private/SteamVRInputDevice.cpp
@@ -2747,7 +2747,7 @@ void FSteamVRInputDevice::GenerateActionManifest(bool GenerateActions, bool Gene
 						YesNoToAll = FMessageDialog::Open(EAppMsgType::YesNoYesAllNoAll, FText::Format(LOCTEXT("BindingFileAlreadyExists", "Your binding file ({0}) was last edited by {1} do you want to overwrite the changes? You will lose any changes you made outside of the editor!"), FText::FromString(BindingFile), FText::FromString(LastEdited)));
 					}
 
-					if (YesNoToAll == EAppReturnType::No || YesNoToAll == EAppReturnType::NoAll)
+					if (YesNoToAll != EAppReturnType::Yes && YesNoToAll != EAppReturnType::YesAll)
 					{
 						// Prevent this binding from being regenerated
 						bIsGenerated = true;
