Fixed crash on editor startup with SteamVR/SteamVR input. Crash was caused by input path in the plugin unloading/reloading the SteamVR runtime, which could cause a race-condition resulting in null internal pointers in the runtime and a crash.
(on behalf of Jules Blok)

diff --git a/Engine/Plugins/Runtime/Steam/SteamVR/Source/SteamVRInputDevice/Private/SteamVRInputDevice.cpp b/Engine/Plugins/Runtime/Steam/SteamVR/Source/SteamVRInputDevice/Private/SteamVRInputDevice.cpp
--- a/Engine/Plugins/Runtime/Steam/SteamVR/Source/SteamVRInputDevice/Private/SteamVRInputDevice.cpp
+++ b/Engine/Plugins/Runtime/Steam/SteamVR/Source/SteamVRInputDevice/Private/SteamVRInputDevice.cpp
@@ -110,23 +110,6 @@ FSteamVRInputDevice::FSteamVRInputDevice(const TSharedRef<FGenericApplicationMes
 
 #if WITH_EDITOR
 	GenerateActionManifest();
-
-	if (VRSettings() != nullptr)
-	{
-		EVRInitError SteamVRInitError = VRInitError_Driver_NotLoaded;
-		IVRSystem* SteamVRSystem = vr::VR_Init(&SteamVRInitError, vr::VRApplication_Overlay);
-
-		if (SteamVRInitError == VRInitError_None)
-		{
-			EVRSettingsError BindingFlagError = VRSettingsError_None;
-			VRSettings()->SetBool(k_pch_SteamVR_Section, k_pch_SteamVR_DebugInputBinding, true, &BindingFlagError);
-			UE_LOG(LogSteamVRInputDevice, Display, TEXT("[STEAMVR INPUT] Enable SteamVR Input Developer Mode: %s"), *FString(UTF8_TO_TCHAR(VRSettings()->GetSettingsErrorNameFromEnum(BindingFlagError))));
-			VRSettings()->SetBool(k_pch_SteamVR_Section, k_pch_SteamVR_DebugInput, true, &BindingFlagError);
-			UE_LOG(LogSteamVRInputDevice, Display, TEXT("[STEAMVR INPUT] Enable SteamVR Debug Input: %s"), *FString(UTF8_TO_TCHAR(VRSettings()->GetSettingsErrorNameFromEnum(BindingFlagError))));
-
-			VR_Shutdown();
-		}
-	}
 #endif
 
 	IModularFeatures::Get().RegisterModularFeature(GetModularFeatureName(), this);
diff --git a/Engine/Plugins/Runtime/Steam/SteamVR/Source/SteamVRInputDevice/Private/SteamVRInputDeviceModule.cpp b/Engine/Plugins/Runtime/Steam/SteamVR/Source/SteamVRInputDevice/Private/SteamVRInputDeviceModule.cpp
--- a/Engine/Plugins/Runtime/Steam/SteamVR/Source/SteamVRInputDevice/Private/SteamVRInputDeviceModule.cpp
+++ b/Engine/Plugins/Runtime/Steam/SteamVR/Source/SteamVRInputDevice/Private/SteamVRInputDeviceModule.cpp
@@ -71,6 +71,24 @@ class FSteamVRInputDeviceModule : public ISteamVRInputDeviceModule
 void FSteamVRInputDeviceModule::StartupModule()
 {
 	IModularFeatures::Get().RegisterModularFeature(GetModularFeatureName(), this);
+
+#if WITH_EDITOR
+	EVRInitError SteamVRInitErrorOverlay = VRInitError_Driver_NotLoaded;
+	IVRSystem* SteamVRSystemOverlay = vr::VR_Init(&SteamVRInitErrorOverlay, vr::VRApplication_Overlay);
+
+	if (SteamVRInitErrorOverlay == VRInitError_None)
+	{
+		if (VRSettings() != nullptr)
+		{
+			EVRSettingsError BindingFlagError = VRSettingsError_None;
+			VRSettings()->SetBool(k_pch_SteamVR_Section, k_pch_SteamVR_DebugInputBinding, true, &BindingFlagError);
+			UE_LOG(LogTemp, Display, TEXT("[STEAMVR INPUT] Enable SteamVR Input Developer Mode: %s"), *FString(UTF8_TO_TCHAR(VRSettings()->GetSettingsErrorNameFromEnum(BindingFlagError))));
+			VRSettings()->SetBool(k_pch_SteamVR_Section, k_pch_SteamVR_DebugInput, true, &BindingFlagError);
+			UE_LOG(LogTemp, Display, TEXT("[STEAMVR INPUT] Enable SteamVR Debug Input: %s"), *FString(UTF8_TO_TCHAR(VRSettings()->GetSettingsErrorNameFromEnum(BindingFlagError))));
+		}
+		vr::VR_Shutdown();
+	}
+#endif
 }
 
 void FSteamVRInputDeviceModule::ShutdownModule()
