Merging CL 10653404 from //UE4/Dev-VR to Release-4.24 (//UE4/Release-4.24)

Merged pull request #6428:

(1) Assert crash/check when enabling debug options for SteamVR 1.9.7 and above
(2) Fix issue when an action is inadvertently triggered (can happen when a specific controller adds a mapping to multiple actions where another controller isn't mapped to exactly the same actions)
(3) Skip bIsGenerated checks for other valid controllers when one is found.
(4) Cleanly shutdown when Marketplace or custom GitHub plugin build is active in the project
(5) Trigger axis values only when there's a change from the previous value
(6) Fix Cosmos axis events showing up as bindable in blueprint

diff --git a/Engine/Plugins/Runtime/Steam/SteamVR/Source/SteamVRInputDevice/Private/SteamVRInputDevice.cpp b/Engine/Plugins/Runtime/Steam/SteamVR/Source/SteamVRInputDevice/Private/SteamVRInputDevice.cpp
--- a/Engine/Plugins/Runtime/Steam/SteamVR/Source/SteamVRInputDevice/Private/SteamVRInputDevice.cpp
+++ b/Engine/Plugins/Runtime/Steam/SteamVR/Source/SteamVRInputDevice/Private/SteamVRInputDevice.cpp
@@ -123,7 +123,6 @@ FSteamVRInputDevice::~FSteamVRInputDevice()
 void FSteamVRInputDevice::InitSteamVRSystem()
 {
 	//UE_LOG(LogTemp, Warning, TEXT("Attempting to load Steam VR System..."));
-
 	if (!VRSystem())
 	{
 		// Initialize OpenVR
@@ -1169,13 +1168,13 @@ void FSteamVRInputDevice::InitControllerKeys()
 	EKeys::AddKey(FKeyDetails(CosmosKeys::Cosmos_Left_Y_Touch, LOCTEXT("Cosmos_Left_Y_Touch", "Cosmos (L) Y Touch"), FKeyDetails::GamepadKey | FKeyDetails::NotBlueprintBindableKey, "Cosmos"));
 	EKeys::AddKey(FKeyDetails(CosmosKeys::Cosmos_Left_Menu_Click, LOCTEXT("Cosmos_Left_Menu_Click", "Cosmos (L) Menu"), FKeyDetails::GamepadKey | FKeyDetails::NotBlueprintBindableKey, "Cosmos"));
 	EKeys::AddKey(FKeyDetails(CosmosKeys::Cosmos_Left_Grip_Click, LOCTEXT("Cosmos_Left_Grip_Click", "Cosmos (L) Grip"), FKeyDetails::GamepadKey | FKeyDetails::NotBlueprintBindableKey, "Cosmos"));
-	EKeys::AddKey(FKeyDetails(CosmosKeys::Cosmos_Left_Grip_Axis, LOCTEXT("Cosmos_Left_Grip_Axis", "Cosmos (L) Grip Axis"), FKeyDetails::GamepadKey | FKeyDetails::FloatAxis, "Cosmos"));
+	EKeys::AddKey(FKeyDetails(CosmosKeys::Cosmos_Left_Grip_Axis, LOCTEXT("Cosmos_Left_Grip_Axis", "Cosmos (L) Grip Axis"), FKeyDetails::GamepadKey | FKeyDetails::FloatAxis | FKeyDetails::NotBlueprintBindableKey, "Cosmos"));
 	EKeys::AddKey(FKeyDetails(CosmosKeys::Cosmos_Left_Trigger_Click, LOCTEXT("Cosmos_Left_Trigger_Click", "Cosmos (L) Trigger"), FKeyDetails::GamepadKey | FKeyDetails::NotBlueprintBindableKey, "Cosmos"));
-	EKeys::AddKey(FKeyDetails(CosmosKeys::Cosmos_Left_Trigger_Axis, LOCTEXT("Cosmos_Left_Trigger_Axis", "Cosmos (L) Trigger Axis"), FKeyDetails::GamepadKey | FKeyDetails::FloatAxis, "Cosmos"));
+	EKeys::AddKey(FKeyDetails(CosmosKeys::Cosmos_Left_Trigger_Axis, LOCTEXT("Cosmos_Left_Trigger_Axis", "Cosmos (L) Trigger Axis"), FKeyDetails::GamepadKey | FKeyDetails::FloatAxis | FKeyDetails::NotBlueprintBindableKey, "Cosmos"));
 	EKeys::AddKey(FKeyDetails(CosmosKeys::Cosmos_Left_Trigger_Touch, LOCTEXT("Cosmos_Left_Trigger_Touch", "Cosmos (L) Trigger Touch"), FKeyDetails::GamepadKey | FKeyDetails::NotBlueprintBindableKey, "Cosmos"));
-	EKeys::AddKey(FKeyDetails(CosmosKeys::Cosmos_Left_Thumbstick_Vector, LOCTEXT("Cosmos_Left_Thumbstick_Vector", "Cosmos (L) Thumbstick Vector"), FKeyDetails::GamepadKey | FKeyDetails::VectorAxis, "Cosmos"));
-	EKeys::AddKey(FKeyDetails(CosmosKeys::Cosmos_Left_Thumbstick_X, LOCTEXT("Cosmos_Left_Thumbstick_X", "Cosmos (L) Thumbstick X"), FKeyDetails::GamepadKey | FKeyDetails::FloatAxis, "Cosmos"));
-	EKeys::AddKey(FKeyDetails(CosmosKeys::Cosmos_Left_Thumbstick_Y, LOCTEXT("Cosmos_Left_Thumbstick_Y", "Cosmos (L) Thumbstick Y"), FKeyDetails::GamepadKey | FKeyDetails::FloatAxis, "Cosmos"));
+	EKeys::AddKey(FKeyDetails(CosmosKeys::Cosmos_Left_Thumbstick_Vector, LOCTEXT("Cosmos_Left_Thumbstick_Vector", "Cosmos (L) Thumbstick Vector"), FKeyDetails::GamepadKey | FKeyDetails::VectorAxis | FKeyDetails::NotBlueprintBindableKey, "Cosmos"));
+	EKeys::AddKey(FKeyDetails(CosmosKeys::Cosmos_Left_Thumbstick_X, LOCTEXT("Cosmos_Left_Thumbstick_X", "Cosmos (L) Thumbstick X"), FKeyDetails::GamepadKey | FKeyDetails::FloatAxis | FKeyDetails::NotBlueprintBindableKey, "Cosmos"));
+	EKeys::AddKey(FKeyDetails(CosmosKeys::Cosmos_Left_Thumbstick_Y, LOCTEXT("Cosmos_Left_Thumbstick_Y", "Cosmos (L) Thumbstick Y"), FKeyDetails::GamepadKey | FKeyDetails::FloatAxis | FKeyDetails::NotBlueprintBindableKey, "Cosmos"));
 	EKeys::AddKey(FKeyDetails(CosmosKeys::Cosmos_Left_Thumbstick_Click, LOCTEXT("Cosmos_Left_Thumbstick_Click", "Cosmos (L) Thumbstick"), FKeyDetails::GamepadKey | FKeyDetails::NotBlueprintBindableKey, "Cosmos"));
 	EKeys::AddKey(FKeyDetails(CosmosKeys::Cosmos_Left_Bumper_Click, LOCTEXT("Cosmos_Left_Bumper_Click", "Cosmos (L) Bumper"), FKeyDetails::GamepadKey | FKeyDetails::NotBlueprintBindableKey, "Cosmos"));
 	EKeys::AddKey(FKeyDetails(CosmosKeys::Cosmos_Left_Thumbstick_Touch, LOCTEXT("Cosmos_Left_Thumbstick_Touch", "Cosmos (L) Thumbstick Touch"), FKeyDetails::GamepadKey | FKeyDetails::NotBlueprintBindableKey, "Cosmos"));
@@ -1185,13 +1184,13 @@ void FSteamVRInputDevice::InitControllerKeys()
 	EKeys::AddKey(FKeyDetails(CosmosKeys::Cosmos_Right_B_Touch, LOCTEXT("Cosmos_Right_B_Touch", "Cosmos (R) B Touch"), FKeyDetails::GamepadKey | FKeyDetails::NotBlueprintBindableKey, "Cosmos"));
 	EKeys::AddKey(FKeyDetails(CosmosKeys::Cosmos_Right_System_Click, LOCTEXT("Cosmos_Right_System_Click", "Cosmos (R) System"), FKeyDetails::GamepadKey | FKeyDetails::NotBlueprintBindableKey, "Cosmos"));
 	EKeys::AddKey(FKeyDetails(CosmosKeys::Cosmos_Right_Grip_Click, LOCTEXT("Cosmos_Right_Grip_Click", "Cosmos (R) Grip"), FKeyDetails::GamepadKey | FKeyDetails::NotBlueprintBindableKey, "Cosmos"));
-	EKeys::AddKey(FKeyDetails(CosmosKeys::Cosmos_Right_Grip_Axis, LOCTEXT("Cosmos_Right_Grip_Axis", "Cosmos (R) Grip Axis"), FKeyDetails::GamepadKey | FKeyDetails::FloatAxis, "Cosmos"));
+	EKeys::AddKey(FKeyDetails(CosmosKeys::Cosmos_Right_Grip_Axis, LOCTEXT("Cosmos_Right_Grip_Axis", "Cosmos (R) Grip Axis"), FKeyDetails::GamepadKey | FKeyDetails::FloatAxis | FKeyDetails::NotBlueprintBindableKey, "Cosmos"));
 	EKeys::AddKey(FKeyDetails(CosmosKeys::Cosmos_Right_Trigger_Click, LOCTEXT("Cosmos_Right_Trigger_Click", "Cosmos (R) Trigger"), FKeyDetails::GamepadKey | FKeyDetails::NotBlueprintBindableKey, "Cosmos"));
-	EKeys::AddKey(FKeyDetails(CosmosKeys::Cosmos_Right_Trigger_Axis, LOCTEXT("Cosmos_Right_Trigger_Axis", "Cosmos (R) Trigger Axis"), FKeyDetails::GamepadKey | FKeyDetails::FloatAxis, "Cosmos"));
+	EKeys::AddKey(FKeyDetails(CosmosKeys::Cosmos_Right_Trigger_Axis, LOCTEXT("Cosmos_Right_Trigger_Axis", "Cosmos (R) Trigger Axis"), FKeyDetails::GamepadKey | FKeyDetails::FloatAxis | FKeyDetails::NotBlueprintBindableKey, "Cosmos"));
 	EKeys::AddKey(FKeyDetails(CosmosKeys::Cosmos_Right_Trigger_Touch, LOCTEXT("Cosmos_Right_Trigger_Touch", "Cosmos (R) Trigger Touch"), FKeyDetails::GamepadKey | FKeyDetails::NotBlueprintBindableKey, "Cosmos"));
-	EKeys::AddKey(FKeyDetails(CosmosKeys::Cosmos_Right_Thumbstick_Vector, LOCTEXT("Cosmos_Right_Thumbstick_Vector", "Cosmos (R) Thumbstick Vector"), FKeyDetails::GamepadKey | FKeyDetails::VectorAxis, "Cosmos"));
-	EKeys::AddKey(FKeyDetails(CosmosKeys::Cosmos_Right_Thumbstick_X, LOCTEXT("Cosmos_Right_Thumbstick_X", "Cosmos (R) Thumbstick X"), FKeyDetails::GamepadKey | FKeyDetails::FloatAxis, "Cosmos"));
-	EKeys::AddKey(FKeyDetails(CosmosKeys::Cosmos_Right_Thumbstick_Y, LOCTEXT("Cosmos_Right_Thumbstick_Y", "Cosmos (R) Thumbstick Y"), FKeyDetails::GamepadKey | FKeyDetails::FloatAxis, "Cosmos"));
+	EKeys::AddKey(FKeyDetails(CosmosKeys::Cosmos_Right_Thumbstick_Vector, LOCTEXT("Cosmos_Right_Thumbstick_Vector", "Cosmos (R) Thumbstick Vector"), FKeyDetails::GamepadKey | FKeyDetails::VectorAxis | FKeyDetails::NotBlueprintBindableKey, "Cosmos"));
+	EKeys::AddKey(FKeyDetails(CosmosKeys::Cosmos_Right_Thumbstick_X, LOCTEXT("Cosmos_Right_Thumbstick_X", "Cosmos (R) Thumbstick X"), FKeyDetails::GamepadKey | FKeyDetails::FloatAxis | FKeyDetails::NotBlueprintBindableKey, "Cosmos"));
+	EKeys::AddKey(FKeyDetails(CosmosKeys::Cosmos_Right_Thumbstick_Y, LOCTEXT("Cosmos_Right_Thumbstick_Y", "Cosmos (R) Thumbstick Y"), FKeyDetails::GamepadKey | FKeyDetails::FloatAxis | FKeyDetails::NotBlueprintBindableKey, "Cosmos"));
 	EKeys::AddKey(FKeyDetails(CosmosKeys::Cosmos_Right_Thumbstick_Click, LOCTEXT("Cosmos_Right_Thumbstick_Click", "Cosmos (R) Thumbstick"), FKeyDetails::GamepadKey | FKeyDetails::NotBlueprintBindableKey, "Cosmos"));
 	EKeys::AddKey(FKeyDetails(CosmosKeys::Cosmos_Right_Thumbstick_Touch, LOCTEXT("Cosmos_Right_Thumbstick_Touch", "Cosmos (R) Thumbstick Touch"), FKeyDetails::GamepadKey | FKeyDetails::NotBlueprintBindableKey, "Cosmos"));
 	EKeys::AddKey(FKeyDetails(CosmosKeys::Cosmos_Right_Bumper_Click, LOCTEXT("Cosmos_Right_Bumper_Click", "Cosmos (R) Bumper"), FKeyDetails::GamepadKey | FKeyDetails::NotBlueprintBindableKey, "Cosmos"));
@@ -1231,6 +1230,7 @@ void FSteamVRInputDevice::OnVREditingModeEnter()
 
 	InputSettings->AddActionMapping(FInputActionKeyMapping("VREditor_WorldMovementL", EKeys::MotionController_Left_Grip1), false);
 	InputSettings->AddActionMapping(FInputActionKeyMapping("VREditor_TouchL", EKeys::Vive_Left_Trackpad_Touch), false);
+	InputSettings->AddAxisMapping(FInputAxisKeyMapping("VREditor_TriggerL", EKeys::MotionController_Left_Trigger), false);
 	InputSettings->AddAxisMapping(FInputAxisKeyMapping("VREditor_TriggerAxisL", EKeys::MotionController_Left_TriggerAxis), false);
 	InputSettings->AddAxisMapping(FInputAxisKeyMapping("VREditor_TrackpadPositionL_X", EKeys::MotionController_Left_Thumbstick_X), false);
 	InputSettings->AddAxisMapping(FInputAxisKeyMapping("VREditor_TrackpadPositionL_Y", EKeys::MotionController_Left_Thumbstick_Y), false);
@@ -1245,6 +1245,7 @@ void FSteamVRInputDevice::OnVREditingModeEnter()
 
 	InputSettings->AddActionMapping(FInputActionKeyMapping("VREditor_WorldMovementR", EKeys::MotionController_Right_Grip1), false);
 	InputSettings->AddActionMapping(FInputActionKeyMapping("VREditor_TouchR", EKeys::Vive_Right_Trackpad_Touch), false);
+	InputSettings->AddAxisMapping(FInputAxisKeyMapping("VREditor_TriggerR", EKeys::MotionController_Right_Trigger), false);
 	InputSettings->AddAxisMapping(FInputAxisKeyMapping("VREditor_TriggerAxisR", EKeys::MotionController_Right_TriggerAxis), false);
 	InputSettings->AddAxisMapping(FInputAxisKeyMapping("VREditor_TrackpadPositionR_X", EKeys::MotionController_Right_Thumbstick_X), false);
 	InputSettings->AddAxisMapping(FInputAxisKeyMapping("VREditor_TrackpadPositionR_Y", EKeys::MotionController_Right_Thumbstick_Y), false);
@@ -2758,6 +2759,7 @@ void FSteamVRInputDevice::GenerateActionManifest(bool GenerateActions, bool Gene
 				if (DefaultControllerType.Name == FName(*ControllerType))
 				{
 					DefaultControllerType.bIsGenerated = bIsGenerated;
+					break;
 				}
 			}
 		}
@@ -2822,9 +2824,11 @@ void FSteamVRInputDevice::GenerateActionManifest(bool GenerateActions, bool Gene
 		RegisterApplication(ManifestPath);
 	}
 
-	// Update action events
+	TMultiMap<FName, FName> KeyToActionsMap;
 	ActionEvents.Empty();
 	bool bAlreadyExists = false;
+
+	// Fill-in unique Action Events that will be processed per tick
 	for (FSteamVRInputAction& InputAction : Actions)
 	{
 		// Check if we've already got a similar action event as we won't need the flat action list for processing controller events
@@ -2838,10 +2842,51 @@ void FSteamVRInputDevice::GenerateActionManifest(bool GenerateActions, bool Gene
 			}
 		}
 
+		// Add unique action handles to action events
 		if (!bAlreadyExists && InputAction.Handle != k_ulInvalidActionHandle)
 		{
 			ActionEvents.Add(FSteamVRInputAction(InputAction));
 		}
+
+		// Add key and action mapping
+		if (InputAction.KeyX != NAME_None)
+		{
+			KeyToActionsMap.AddUnique(InputAction.KeyX, InputAction.Name);
+		}
+	}
+
+	
+	// Find keys that trigger multiple actions
+	TArray<FName> ActionKeys;
+	KeyToActionsMap.GetKeys(ActionKeys);
+	for (auto Key : ActionKeys)
+	{
+		TArray<FName> MappedActions;
+		KeyToActionsMap.MultiFind(Key, MappedActions);
+		if (MappedActions.Num() > 1)
+		{
+			// Remove key-value pairing for keys that trigger multiple actions in map
+			KeyToActionsMap.Remove(Key);
+		}
+	}
+
+	// Finalize keys that will be used to trigger actions per tick
+	KeyToActionsMap.Shrink();
+	for (FSteamVRInputAction& ActionEvent : ActionEvents)
+	{
+		// Skip actions with no key mapping (poses, haptics, etc)
+		if (ActionEvent.KeyX != NAME_None)
+		{
+			// Check if this action is in the KeyToActionsMap map
+			const FName* PriorityKey = KeyToActionsMap.FindKey(ActionEvent.Name);
+			if (PriorityKey != nullptr)
+			{
+				// Set the key in the KeyToActionsMap as the default key that will be used to trigger actions
+				ActionEvent.KeyX = *PriorityKey;
+			}
+		}
+
+		//UE_LOG(LogSteamVRInputDevice, Warning, TEXT("Action [%s] mapped to key [%s]"), *ActionEvent.Name.ToString(), *ActionEvent.KeyX.ToString());
 	}
 }
 
@@ -3607,19 +3652,19 @@ void FSteamVRInputDevice::ProcessActionEvents(FSteamVRInputActionSet SteamVRInpu
 				// Update the active origin
 				Action.ActiveOrigin = AnalogData.activeOrigin;
 
-				if (Action.KeyX != NAME_None)
+				if (Action.KeyX != NAME_None && (AnalogData.deltaX > KINDA_SMALL_NUMBER || Action.Name.IsEqual(FName(CONTROLLER_BINDING_PATH))))
 				{
 					// Test what we're receiving from SteamVR
-					//UE_LOG(LogTemp, Warning, TEXT("Handle %s KeyX %s X-Value [%f]"), *Action.Path, *FoundKeyX.GetFName().ToString(), AnalogData.x);
+					//UE_LOG(LogTemp, Warning, TEXT("Handle %s KeyX %s X-Value [%f]"), *Action.Path, *Action.KeyX.ToString(), AnalogData.x);
 
 					Action.Value.X = AnalogData.x; // ActionCount;
 					MessageHandler->OnControllerAnalog(Action.KeyX, 0, Action.Value.X);
 				}
 
-				if (Action.KeyY != NAME_None)
+				if (Action.KeyY != NAME_None && (AnalogData.deltaY > KINDA_SMALL_NUMBER))
 				{
 					// Test what we're receiving from SteamVR
-					//UE_LOG(LogTemp, Warning, TEXT("Handle %s KeyY %s Y-Value {%f}"), *Action.Path, *FoundKeyY.GetFName().ToString(), AnalogData.y);
+					//UE_LOG(LogTemp, Warning, TEXT("Handle %s KeyY %s Y-Value {%f}"), *Action.Path, *Action.KeyY.ToString(), AnalogData.y);
 					
 					FString KeyString = Action.KeyY.ToString();
 					if (KeyString.Contains(TEXT("MotionController")) && KeyString.Contains(TEXT("_Y")))
diff --git a/Engine/Plugins/Runtime/Steam/SteamVR/Source/SteamVRInputDevice/Private/SteamVRInputDeviceModule.cpp b/Engine/Plugins/Runtime/Steam/SteamVR/Source/SteamVRInputDevice/Private/SteamVRInputDeviceModule.cpp
--- a/Engine/Plugins/Runtime/Steam/SteamVR/Source/SteamVRInputDevice/Private/SteamVRInputDeviceModule.cpp
+++ b/Engine/Plugins/Runtime/Steam/SteamVR/Source/SteamVRInputDevice/Private/SteamVRInputDeviceModule.cpp
@@ -86,9 +86,19 @@ void FSteamVRInputDeviceModule::StartupModule()
 			VRSettings()->SetBool(k_pch_SteamVR_Section, k_pch_SteamVR_DebugInput, true, &BindingFlagError);
 			UE_LOG(LogTemp, Display, TEXT("[STEAMVR INPUT] Enable SteamVR Debug Input: %s"), *FString(UTF8_TO_TCHAR(VRSettings()->GetSettingsErrorNameFromEnum(BindingFlagError))));
 		}
-		vr::VR_Shutdown();
+		VR_Shutdown();
 	}
 #endif
+
+	// Unload engine integrated controller modules if Valve's Input Plugin is present
+	FModuleManager& ModuleManager = FModuleManager::Get();
+	IModuleInterface* ValveInputOverride = ModuleManager.GetModule(FName("SteamVR_Input"));
+
+	if (ValveInputOverride != nullptr)
+	{
+		UE_LOG(LogTemp, Log, TEXT("[SteamVR Input] Found Valve Input Plugin. Integrated version of SteamVR Input will not be loaded."));
+		IModularFeatures::Get().UnregisterModularFeature(GetModularFeatureName(), this);
+	}
 }
 
 void FSteamVRInputDeviceModule::ShutdownModule()
