Move SteamVRInput headers that shouldn't be in Public to Private

#rnx

diff --git a/Engine/Plugins/Runtime/Steam/SteamVR/Source/SteamVREditor/Private/SteamVREditor.cpp b/Engine/Plugins/Runtime/Steam/SteamVR/Source/SteamVREditor/Private/SteamVREditor.cpp
--- a/Engine/Plugins/Runtime/Steam/SteamVR/Source/SteamVREditor/Private/SteamVREditor.cpp
+++ b/Engine/Plugins/Runtime/Steam/SteamVR/Source/SteamVREditor/Private/SteamVREditor.cpp
@@ -35,7 +35,6 @@ POSSIBILITY OF SUCH DAMAGE.
 #include "Features/IModularFeatures.h"
 #include "SteamVREditorStyle.h"
 #include "SteamVREditorCommands.h"
-#include "SteamVRControllerKeys.h"
 #include "LevelEditor.h"
 #include "Engine/Engine.h"
 #include "IXRTrackingSystem.h"
diff --git a/Engine/Plugins/Runtime/Steam/SteamVR/Source/SteamVREditor/Private/SteamVREditorCommands.cpp b/Engine/Plugins/Runtime/Steam/SteamVR/Source/SteamVREditor/Private/SteamVREditorCommands.cpp
new file mode 100644
--- /dev/null
+++ b/Engine/Plugins/Runtime/Steam/SteamVR/Source/SteamVREditor/Private/SteamVREditorCommands.cpp
@@ -0,0 +1,46 @@
+/*
+Copyright 2019 Valve Corporation under https://opensource.org/licenses/BSD-3-Clause
+This code includes modifications by Epic Games.  Modifications (c) Epic Games, Inc.
+
+Redistribution and use in source and binary forms, with or without modification,
+are permitted provided that the following conditions are met:
+
+1. Redistributions of source code must retain the above copyright notice, this
+   list of conditions and the following disclaimer.
+
+2. Redistributions in binary form must reproduce the above copyright notice,
+   this list of conditions and the following disclaimer in the documentation
+   and/or other materials provided with the distribution.
+
+3. Neither the name of the copyright holder nor the names of its contributors
+   may be used to endorse or promote products derived from this software
+   without specific prior written permission.
+
+THIS SOFTWARE IS PROVIDED BY THE COPYRIGHT HOLDERS AND CONTRIBUTORS "AS IS" AND
+ANY EXPRESS OR IMPLIED WARRANTIES, INCLUDING, BUT NOT LIMITED TO, THE IMPLIED
+WARRANTIES OF MERCHANTABILITY AND FITNESS FOR A PARTICULAR PURPOSE ARE DISCLAIMED.
+IN NO EVENT SHALL THE COPYRIGHT HOLDER OR CONTRIBUTORS BE LIABLE FOR ANY DIRECT,
+INDIRECT, INCIDENTAL, SPECIAL, EXEMPLARY, OR CONSEQUENTIAL DAMAGES (INCLUDING,
+BUT NOT LIMITED TO, PROCUREMENT OF SUBSTITUTE GOODS OR SERVICES; LOSS OF USE, DATA,
+OR PROFITS; OR BUSINESS INTERRUPTION) HOWEVER CAUSED AND ON ANY THEORY OF LIABILITY,
+WHETHER IN CONTRACT, STRICT LIABILITY, OR TORT (INCLUDING NEGLIGENCE OR OTHERWISE)
+ARISING IN ANY WAY OUT OF THE USE OF THIS SOFTWARE, EVEN IF ADVISED OF THE
+POSSIBILITY OF SUCH DAMAGE.
+*/
+
+#include "SteamVREditorCommands.h"
+#include "SteamVREditor.h"
+
+#define LOCTEXT_NAMESPACE "FSteamVREditorModule"
+
+void FSteamVREditorCommands::RegisterCommands()
+{
+	UI_COMMAND(PluginAction, "SteamVREditor", "Execute SteamVREditor action", EUserInterfaceActionType::Button, FInputGesture());
+	UI_COMMAND(JsonActionManifest, "Regenerate Action Manifest", "Regenerate Action Manifest", EUserInterfaceActionType::Button, FInputGesture());
+	UI_COMMAND(JsonControllerBindings, "Regenerate Controller Bindings", "Regenerate Controller Bindings", EUserInterfaceActionType::Button, FInputGesture());
+	UI_COMMAND(ReloadActionManifest, "Reload Action Manifest", "Reload Action Manifest", EUserInterfaceActionType::Button, FInputGesture());
+	UI_COMMAND(LaunchBindingsURL, "Launch SteamVR Bindings Dashboard", "Launch SteamVR Bindings Dashboard", EUserInterfaceActionType::Button, FInputGesture());
+	UI_COMMAND(AddSampleInputs, "Add Sample Inputs", "Add Sample Inputs", EUserInterfaceActionType::Button, FInputGesture());
+}
+
+#undef LOCTEXT_NAMESPACE
diff --git a/Engine/Plugins/Runtime/Steam/SteamVR/Source/SteamVREditor/Private/SteamVREditorCommands.h b/Engine/Plugins/Runtime/Steam/SteamVR/Source/SteamVREditor/Private/SteamVREditorCommands.h
new file mode 100644
--- /dev/null
+++ b/Engine/Plugins/Runtime/Steam/SteamVR/Source/SteamVREditor/Private/SteamVREditorCommands.h
@@ -0,0 +1,54 @@
+/*
+Copyright 2019 Valve Corporation under https://opensource.org/licenses/BSD-3-Clause
+This code includes modifications by Epic Games.  Modifications (c) Epic Games, Inc.
+
+Redistribution and use in source and binary forms, with or without modification,
+are permitted provided that the following conditions are met:
+
+1. Redistributions of source code must retain the above copyright notice, this
+   list of conditions and the following disclaimer.
+
+2. Redistributions in binary form must reproduce the above copyright notice,
+   this list of conditions and the following disclaimer in the documentation
+   and/or other materials provided with the distribution.
+
+3. Neither the name of the copyright holder nor the names of its contributors
+   may be used to endorse or promote products derived from this software
+   without specific prior written permission.
+
+THIS SOFTWARE IS PROVIDED BY THE COPYRIGHT HOLDERS AND CONTRIBUTORS "AS IS" AND
+ANY EXPRESS OR IMPLIED WARRANTIES, INCLUDING, BUT NOT LIMITED TO, THE IMPLIED
+WARRANTIES OF MERCHANTABILITY AND FITNESS FOR A PARTICULAR PURPOSE ARE DISCLAIMED.
+IN NO EVENT SHALL THE COPYRIGHT HOLDER OR CONTRIBUTORS BE LIABLE FOR ANY DIRECT,
+INDIRECT, INCIDENTAL, SPECIAL, EXEMPLARY, OR CONSEQUENTIAL DAMAGES (INCLUDING,
+BUT NOT LIMITED TO, PROCUREMENT OF SUBSTITUTE GOODS OR SERVICES; LOSS OF USE, DATA,
+OR PROFITS; OR BUSINESS INTERRUPTION) HOWEVER CAUSED AND ON ANY THEORY OF LIABILITY,
+WHETHER IN CONTRACT, STRICT LIABILITY, OR TORT (INCLUDING NEGLIGENCE OR OTHERWISE)
+ARISING IN ANY WAY OUT OF THE USE OF THIS SOFTWARE, EVEN IF ADVISED OF THE
+POSSIBILITY OF SUCH DAMAGE.
+*/
+
+#pragma once
+
+#include "SteamVREditorStyle.h"
+
+class FSteamVREditorCommands : public TCommands<FSteamVREditorCommands>
+{
+public:
+
+	FSteamVREditorCommands()
+		: TCommands<FSteamVREditorCommands>(TEXT("SteamVREditor"), NSLOCTEXT("Contexts", "SteamVREditor", "SteamVREditor Plugin"), NAME_None, TEXT("SteamVREditor.Common.Icon"))
+	{
+	}
+
+	// TCommands<> interface
+	virtual void RegisterCommands() override;
+
+public:
+	TSharedPtr<FUICommandInfo> PluginAction;
+	TSharedPtr<FUICommandInfo> JsonActionManifest;
+	TSharedPtr<FUICommandInfo> JsonControllerBindings;
+	TSharedPtr<FUICommandInfo> ReloadActionManifest;
+	TSharedPtr<FUICommandInfo> LaunchBindingsURL;
+	TSharedPtr<FUICommandInfo> AddSampleInputs;
+};
diff --git a/Engine/Plugins/Runtime/Steam/SteamVR/Source/SteamVREditor/Private/SteamVREditorStyle.h b/Engine/Plugins/Runtime/Steam/SteamVR/Source/SteamVREditor/Private/SteamVREditorStyle.h
new file mode 100644
--- /dev/null
+++ b/Engine/Plugins/Runtime/Steam/SteamVR/Source/SteamVREditor/Private/SteamVREditorStyle.h
@@ -0,0 +1,64 @@
+/*
+Copyright 2019 Valve Corporation under https://opensource.org/licenses/BSD-3-Clause
+This code includes modifications by Epic Games.  Modifications (c) Epic Games, Inc.
+
+Redistribution and use in source and binary forms, with or without modification,
+are permitted provided that the following conditions are met:
+
+1. Redistributions of source code must retain the above copyright notice, this
+   list of conditions and the following disclaimer.
+
+2. Redistributions in binary form must reproduce the above copyright notice,
+   this list of conditions and the following disclaimer in the documentation
+   and/or other materials provided with the distribution.
+
+3. Neither the name of the copyright holder nor the names of its contributors
+   may be used to endorse or promote products derived from this software
+   without specific prior written permission.
+
+THIS SOFTWARE IS PROVIDED BY THE COPYRIGHT HOLDERS AND CONTRIBUTORS "AS IS" AND
+ANY EXPRESS OR IMPLIED WARRANTIES, INCLUDING, BUT NOT LIMITED TO, THE IMPLIED
+WARRANTIES OF MERCHANTABILITY AND FITNESS FOR A PARTICULAR PURPOSE ARE DISCLAIMED.
+IN NO EVENT SHALL THE COPYRIGHT HOLDER OR CONTRIBUTORS BE LIABLE FOR ANY DIRECT,
+INDIRECT, INCIDENTAL, SPECIAL, EXEMPLARY, OR CONSEQUENTIAL DAMAGES (INCLUDING,
+BUT NOT LIMITED TO, PROCUREMENT OF SUBSTITUTE GOODS OR SERVICES; LOSS OF USE, DATA,
+OR PROFITS; OR BUSINESS INTERRUPTION) HOWEVER CAUSED AND ON ANY THEORY OF LIABILITY,
+WHETHER IN CONTRACT, STRICT LIABILITY, OR TORT (INCLUDING NEGLIGENCE OR OTHERWISE)
+ARISING IN ANY WAY OUT OF THE USE OF THIS SOFTWARE, EVEN IF ADVISED OF THE
+POSSIBILITY OF SUCH DAMAGE.
+*/
+
+#pragma once
+
+#include "Styling/SlateStyle.h"
+#include "Styling/SlateStyleRegistry.h"
+#include "Framework/Application/SlateApplication.h"
+#include "Framework/MultiBox/MultiBoxBuilder.h"
+#include "Framework/Commands/Commands.h"
+#include "Framework/Commands/UIAction.h"
+#include "Framework/Commands/UICommandInfo.h"
+
+/**  */
+class FSteamVREditorStyle
+{
+public:
+	static void Initialize();
+
+	static void Shutdown();
+
+	/** reloads textures used by slate renderer */
+	static void ReloadTextures();
+
+	/** @return The Slate style set */
+	static const ISlateStyle& Get();
+
+	static FName GetStyleSetName();
+
+private:
+
+	static TSharedRef< class FSlateStyleSet > Create();
+
+private:
+
+	static TSharedPtr< class FSlateStyleSet > StyleInstance;
+};
\ No newline at end of file
diff --git a/Engine/Plugins/Runtime/Steam/SteamVR/Source/SteamVREditor/Public/SteamVREditor.h b/Engine/Plugins/Runtime/Steam/SteamVR/Source/SteamVREditor/Public/SteamVREditor.h
--- a/Engine/Plugins/Runtime/Steam/SteamVR/Source/SteamVREditor/Public/SteamVREditor.h
+++ b/Engine/Plugins/Runtime/Steam/SteamVR/Source/SteamVREditor/Public/SteamVREditor.h
@@ -31,8 +31,7 @@ POSSIBILITY OF SUCH DAMAGE.
 #pragma once
 
 #include "Modules/ModuleManager.h"
-#include "../../SteamVRInputDevice/Public/SteamVRInputDevice.h"
-#include "../../SteamVRInputDevice/Public/SteamVRInputDeviceFunctionLibrary.h"
+#include "SteamVRInputDeviceFunctionLibrary.h"
 
 class FToolBarBuilder;
 class FMenuBuilder;
diff --git a/Engine/Plugins/Runtime/Steam/SteamVR/Source/SteamVREditor/Public/SteamVREditorCommands.cpp b/Engine/Plugins/Runtime/Steam/SteamVR/Source/SteamVREditor/Public/SteamVREditorCommands.cpp
deleted file mode 100644
diff --git a/Engine/Plugins/Runtime/Steam/SteamVR/Source/SteamVREditor/Public/SteamVREditorCommands.h b/Engine/Plugins/Runtime/Steam/SteamVR/Source/SteamVREditor/Public/SteamVREditorCommands.h
deleted file mode 100644
diff --git a/Engine/Plugins/Runtime/Steam/SteamVR/Source/SteamVREditor/Public/SteamVREditorStyle.h b/Engine/Plugins/Runtime/Steam/SteamVR/Source/SteamVREditor/Public/SteamVREditorStyle.h
deleted file mode 100644
diff --git a/Engine/Plugins/Runtime/Steam/SteamVR/Source/SteamVRInputDevice/Private/SteamVRControllerKeys.h b/Engine/Plugins/Runtime/Steam/SteamVR/Source/SteamVRInputDevice/Private/SteamVRControllerKeys.h
new file mode 100644
--- /dev/null
+++ b/Engine/Plugins/Runtime/Steam/SteamVR/Source/SteamVRInputDevice/Private/SteamVRControllerKeys.h
@@ -0,0 +1,95 @@
+/*
+Copyright 2019 Valve Corporation under https://opensource.org/licenses/BSD-3-Clause
+This code includes modifications by Epic Games.  Modifications (c) Epic Games, Inc.
+
+Redistribution and use in source and binary forms, with or without modification,
+are permitted provided that the following conditions are met:
+
+1. Redistributions of source code must retain the above copyright notice, this
+   list of conditions and the following disclaimer.
+
+2. Redistributions in binary form must reproduce the above copyright notice,
+   this list of conditions and the following disclaimer in the documentation
+   and/or other materials provided with the distribution.
+
+3. Neither the name of the copyright holder nor the names of its contributors
+   may be used to endorse or promote products derived from this software
+   without specific prior written permission.
+
+THIS SOFTWARE IS PROVIDED BY THE COPYRIGHT HOLDERS AND CONTRIBUTORS "AS IS" AND
+ANY EXPRESS OR IMPLIED WARRANTIES, INCLUDING, BUT NOT LIMITED TO, THE IMPLIED
+WARRANTIES OF MERCHANTABILITY AND FITNESS FOR A PARTICULAR PURPOSE ARE DISCLAIMED.
+IN NO EVENT SHALL THE COPYRIGHT HOLDER OR CONTRIBUTORS BE LIABLE FOR ANY DIRECT,
+INDIRECT, INCIDENTAL, SPECIAL, EXEMPLARY, OR CONSEQUENTIAL DAMAGES (INCLUDING,
+BUT NOT LIMITED TO, PROCUREMENT OF SUBSTITUTE GOODS OR SERVICES; LOSS OF USE, DATA,
+OR PROFITS; OR BUSINESS INTERRUPTION) HOWEVER CAUSED AND ON ANY THEORY OF LIABILITY,
+WHETHER IN CONTRACT, STRICT LIABILITY, OR TORT (INCLUDING NEGLIGENCE OR OTHERWISE)
+ARISING IN ANY WAY OUT OF THE USE OF THIS SOFTWARE, EVEN IF ADVISED OF THE
+POSSIBILITY OF SUCH DAMAGE.
+*/
+
+#pragma once
+#include "InputCoreTypes.h"
+
+namespace GenericKeys
+{
+	const FKey SteamVR_MotionController_None("SteamVR_MotionController_None");
+	const FKey SteamVR_HMD_Proximity("SteamVR_HMD_Proximity");
+}
+
+namespace IndexControllerKeys
+{
+	// Special
+	const FKey ValveIndex_Left_Grip_Grab("ValveIndex_Left_Grip_Grab");
+	const FKey ValveIndex_Left_Pinch_Grab("ValveIndex_Left_Pinch_Grab");
+	const FKey ValveIndex_Right_Grip_Grab("ValveIndex_Right_Grip_Grab");
+	const FKey ValveIndex_Right_Pinch_Grab("ValveIndex_Right_Pinch_Grab");
+}
+
+namespace CosmosKeys
+{
+	// Regular Controller Keys
+	const FKey Cosmos_Left_X_Click("Cosmos_Left_X_Click");
+	const FKey Cosmos_Left_Y_Click("Cosmos_Left_Y_Click");
+	const FKey Cosmos_Left_X_Touch("Cosmos_Left_X_Touch");
+	const FKey Cosmos_Left_Y_Touch("Cosmos_Left_Y_Touch");
+	const FKey Cosmos_Left_Menu_Click("Cosmos_Left_Menu_Click");
+	const FKey Cosmos_Left_Grip_Click("Cosmos_Left_Grip_Click");
+	const FKey Cosmos_Left_Grip_Axis("Cosmos_Left_Grip_Axis");
+	const FKey Cosmos_Left_Trigger_Click("Cosmos_Left_Trigger_Click");
+	const FKey Cosmos_Left_Trigger_Axis("Cosmos_Left_Trigger_Axis");
+	const FKey Cosmos_Left_Trigger_Touch("Cosmos_Left_Trigger_Touch");
+	const FKey Cosmos_Left_Thumbstick_Vector("Cosmos_Left_Thumbstick_Vector");
+	const FKey Cosmos_Left_Thumbstick_X("Cosmos_Left_Thumbstick_X");
+	const FKey Cosmos_Left_Thumbstick_Y("Cosmos_Left_Thumbstick_Y");
+	const FKey Cosmos_Left_Thumbstick_Click("Cosmos_Left_Thumbstick_Click");
+	const FKey Cosmos_Left_Thumbstick_Touch("Cosmos_Left_Thumbstick_Touch");
+	const FKey Cosmos_Left_Bumper_Click("Cosmos_Left_Bumper_Click");
+	const FKey Cosmos_Right_A_Click("Cosmos_Right_A_Click");
+	const FKey Cosmos_Right_B_Click("Cosmos_Right_B_Click");
+	const FKey Cosmos_Right_A_Touch("Cosmos_Right_A_Touch");
+	const FKey Cosmos_Right_B_Touch("Cosmos_Right_B_Touch");
+	const FKey Cosmos_Right_System_Click("Cosmos_Right_System_Click");
+	const FKey Cosmos_Right_Grip_Click("Cosmos_Right_Grip_Click");
+	const FKey Cosmos_Right_Grip_Axis("Cosmos_Right_Grip_Axis");
+	const FKey Cosmos_Right_Trigger_Click("Cosmos_Right_Trigger_Click");
+	const FKey Cosmos_Right_Trigger_Axis("Cosmos_Right_Trigger_Axis");
+	const FKey Cosmos_Right_Trigger_Touch("Cosmos_Right_Trigger_Touch");
+	const FKey Cosmos_Right_Thumbstick_Vector("Cosmos_Right_Thumbstick_Vector");
+	const FKey Cosmos_Right_Thumbstick_X("Cosmos_Right_Thumbstick_X");
+	const FKey Cosmos_Right_Thumbstick_Y("Cosmos_Right_Thumbstick_Y");
+	const FKey Cosmos_Right_Thumbstick_Click("Cosmos_Right_Thumbstick_Click");
+	const FKey Cosmos_Right_Thumbstick_Touch("Cosmos_Right_Thumbstick_Touch");
+	const FKey Cosmos_Right_Bumper_Click("Cosmos_Right_Bumper_Click");
+
+	// Thumbstick Directions
+	const FKey Cosmos_Left_Thumbstick_Up("Cosmos_Left_Thumbstick_Up");
+	const FKey Cosmos_Left_Thumbstick_Down("Cosmos_Left_Thumbstick_Down");
+	const FKey Cosmos_Left_Thumbstick_Left("Cosmos_Left_Thumbstick_Left");
+	const FKey Cosmos_Left_Thumbstick_Right("Cosmos_Left_Thumbstick_Right");
+	const FKey Cosmos_Right_Thumbstick_Up("Cosmos_Right_Thumbstick_Up");
+	const FKey Cosmos_Right_Thumbstick_Down("Cosmos_Right_Thumbstick_Down");
+	const FKey Cosmos_Right_Thumbstick_Left("Cosmos_Right_Thumbstick_Left");
+	const FKey Cosmos_Right_Thumbstick_Right("Cosmos_Right_Thumbstick_Right");
+}
+
diff --git a/Engine/Plugins/Runtime/Steam/SteamVR/Source/SteamVRInputDevice/Private/SteamVRInputDevice.h b/Engine/Plugins/Runtime/Steam/SteamVR/Source/SteamVRInputDevice/Private/SteamVRInputDevice.h
new file mode 100644
--- /dev/null
+++ b/Engine/Plugins/Runtime/Steam/SteamVR/Source/SteamVRInputDevice/Private/SteamVRInputDevice.h
@@ -0,0 +1,488 @@
+/*
+Copyright 2019 Valve Corporation under https://opensource.org/licenses/BSD-3-Clause
+This code includes modifications by Epic Games.  Modifications (c) Epic Games, Inc.
+
+Redistribution and use in source and binary forms, with or without modification,
+are permitted provided that the following conditions are met:
+
+1. Redistributions of source code must retain the above copyright notice, this
+   list of conditions and the following disclaimer.
+
+2. Redistributions in binary form must reproduce the above copyright notice,
+   this list of conditions and the following disclaimer in the documentation
+   and/or other materials provided with the distribution.
+
+3. Neither the name of the copyright holder nor the names of its contributors
+   may be used to endorse or promote products derived from this software
+   without specific prior written permission.
+
+THIS SOFTWARE IS PROVIDED BY THE COPYRIGHT HOLDERS AND CONTRIBUTORS "AS IS" AND
+ANY EXPRESS OR IMPLIED WARRANTIES, INCLUDING, BUT NOT LIMITED TO, THE IMPLIED
+WARRANTIES OF MERCHANTABILITY AND FITNESS FOR A PARTICULAR PURPOSE ARE DISCLAIMED.
+IN NO EVENT SHALL THE COPYRIGHT HOLDER OR CONTRIBUTORS BE LIABLE FOR ANY DIRECT,
+INDIRECT, INCIDENTAL, SPECIAL, EXEMPLARY, OR CONSEQUENTIAL DAMAGES (INCLUDING,
+BUT NOT LIMITED TO, PROCUREMENT OF SUBSTITUTE GOODS OR SERVICES; LOSS OF USE, DATA,
+OR PROFITS; OR BUSINESS INTERRUPTION) HOWEVER CAUSED AND ON ANY THEORY OF LIABILITY,
+WHETHER IN CONTRACT, STRICT LIABILITY, OR TORT (INCLUDING NEGLIGENCE OR OTHERWISE)
+ARISING IN ANY WAY OUT OF THE USE OF THIS SOFTWARE, EVEN IF ADVISED OF THE
+POSSIBILITY OF SUCH DAMAGE.
+*/
+
+#pragma once
+
+#include "IInputDevice.h"
+#include "CoreMinimal.h"
+#include "Runtime/Core/Public/Misc/ConfigCacheIni.h"
+#include "XRMotionControllerBase.h"
+#include "IHapticDevice.h"
+#include "Runtime/Core/Public/GenericPlatform/GenericPlatformProcess.h"
+#include "Runtime/Core/Public/Misc/Paths.h"
+#include "Engine/Engine.h"
+#include "Serialization/JsonReader.h"
+#include "Serialization/JsonSerializer.h"
+#include "SteamVRInputTypes.h"
+#include "SteamVRInputPublic.h"
+
+class STEAMVRINPUTDEVICE_API FSteamVRInputDevice : public IInputDevice, public FXRMotionControllerBase, public IHapticDevice
+{
+public:
+	FSteamVRInputDevice(const TSharedRef<FGenericApplicationMessageHandler>& InMessageHandler);
+
+	~FSteamVRInputDevice();
+
+	// IInput Device Interface
+	virtual void Tick(float DeltaTime) override;
+	virtual void SendControllerEvents() override;
+	virtual void SetMessageHandler(const TSharedRef<FGenericApplicationMessageHandler>& InMessageHandler) override;
+	virtual void SetChannelValue(int32 ControllerId, FForceFeedbackChannelType ChannelType, float Value) override;
+	virtual void SetChannelValues(int32 ControllerId, const FForceFeedbackValues &values) override;
+	virtual bool Exec(UWorld* InWorld, const TCHAR* Cmd, FOutputDevice& Ar) override;
+	// End of IInput Device Interface
+
+	// FXRMotionControllerBase
+	virtual bool GetControllerOrientationAndPosition(const int32 ControllerIndex, const EControllerHand DeviceHand, FRotator& OutOrientation, FVector& OutPosition, float WorldToMetersScale) const;
+	virtual ETrackingStatus GetControllerTrackingStatus(const int32 ControllerIndex, const EControllerHand DeviceHand) const;
+	// End of FXRMotionControllerBase
+	
+	// IMotionController Interface
+	virtual FName GetMotionControllerDeviceTypeName() const override;
+	virtual bool GetHandJointPosition(const FName MotionSource, int jointIndex, FVector& OutPosition) const;
+	// End of IMotionController Interface
+
+	// IHapticDevice Interface
+	IHapticDevice* GetHapticDevice() override { return (IHapticDevice*)this; }
+	virtual void SetHapticFeedbackValues(int32 ControllerId, int32 Hand, const FHapticFeedbackValues& Values) override;
+	virtual void GetHapticFrequencyRange(float& MinFrequency, float& MaxFrequency) const override;
+	virtual float GetHapticAmplitudeScale() const override;
+	// End of IHapticDevice Interface
+
+	/** Initialize the SteamVR System. Will cause a reconnect if one is already active  */
+	void InitSteamVRSystem();
+
+	/**
+	* Retrieve the skeletal input from SteamVR
+	* @param bLeftHand - Whether or not retrieve values for the Left Hand instead of the Right Hand
+	* @param bMirror - Will mirror the pose read from SteamVR to fit the skeleton of the opposite hand
+	* @param MotionRange - Whether to retrieve skeletal anim values with or without controllers
+	* @param OutBoneTransform - The transform for each bone as defined in the SteamVR Skeleton
+	* @param OutBoneTransformCount - The number of bones in the SteamVR Skeleton as reported by the SteamVR Input system
+	* @return Whether or not we successfully retrieved the transform values
+	*/
+	bool GetSkeletalData(bool bLeftHand, bool bMirror, EVRSkeletalMotionRange MotionRange, FTransform* OutBoneTransform, int32 OutBoneTransformCount);
+
+	/**
+	* Retrieve the left hand pose information - position, orientation and velocities
+	* @return Position - Translation from the pose data matrix in UE coordinates
+	* @return Orientation - Orientation derived from the pose data matrix in UE coordinates
+	* @return AngularVelocity - The angular velocity of the hand this frame
+	* @return Velocity - The velocity of the hand this frame
+	*/
+	void GetLeftHandPoseData(FVector& Position, FRotator& Orientation, FVector& AngularVelocity, FVector& Velocity); 
+
+	/**
+	* Retrieve the right hand pose information - position, orientation and velocities
+	* @return Position - Translation from the pose data matrix in UE coordinates
+	* @return Orientation - Orientation derived from the pose data matrix in UE coordinates
+	* @return AngularVelocity - The angular velocity of the hand this frame
+	* @return Velocity - The velocity of the hand this frame
+	*/
+	void GetRightHandPoseData(FVector& Position, FRotator& Orientation, FVector& AngularVelocity, FVector& Velocity);
+	
+	/**
+	* Calculate the UE equivalent orientation and position from a SteamVR PoseData matrix 
+	* @param PoseData - The pose data returned by SteamVR
+	* @return OutPosition - Translation from the pose data
+	* @return OutOrientation - Orientation from the pose data
+	*/
+	void GetUETransform(InputPoseActionData_t PoseData, FVector& OutPosition, FRotator& OutOrientation);
+
+	/** Retrieve skeletal tracking level for all controllers */
+	void GetControllerFidelity();
+
+#if WITH_EDITOR
+	/** Have the action manifest regenerated. Used by the plugin Editor UI */
+	void RegenerateActionManifest();
+
+	/** Have the controller bindings regenerated. Used by the plugin Editor UI */
+	void RegenerateControllerBindings();
+
+	/** Have the manifest and bindings bindings regenerated. Used by a UI callback */
+	void OnActionMappingsChanged();
+
+	/** Add VR editor actions and bindings. Used by a UI callback */
+	void OnVREditingModeEnter();
+
+	/** Remove VR editor actions and bindings. Used by a UI callback */
+	void OnVREditingModeExit();
+
+	/** Reload the action manifest to the SteamVR system. Use if you changed the action manifest manually */
+	void ReloadActionManifest();
+#endif
+
+	/** Whether or not Curls and Splay values for the LEFT HAND are fed to the game every frame */
+	bool bCurlsAndSplaysEnabled_L = true;
+
+	/** Whether or not Curls and Splay values for the RIGHT HAND are fed to the game every frame */
+	bool bCurlsAndSplaysEnabled_R = true;
+
+	/** @deprecated Number of controllers that are currently active in a session */
+	int32 NumControllersMapped;
+
+	/** @deprecated Number of trackers that are currently active in a session */
+	int32 NumTrackersMapped;
+
+	/** @deprecated Mapping between SteamVR Legacy Device Ids and UE4s controller mappings */
+	int32 DeviceToControllerMap[k_unMaxTrackedDeviceCount];
+
+	/** @deprecated Mapping of UE4 devices, hand they are associated with (including "special" for trackers) and legacy SteamVR device Ids */
+	int32 UnrealControllerIdAndHandToDeviceIdMap[SteamVRInputDeviceConstants::MaxUnrealControllers][k_unMaxTrackedDeviceCount];
+
+	/** @deprecated Maximum number of allowed active controllers available to a player */
+	int32 MaxUEHandCount[CONTROLLERS_PER_PLAYER];
+
+	/** Holds the current frame values of a controller, including any SteamVR summary skeletal values */
+	struct FInputDeviceState
+	{
+		EControllerHand Hand;	// The hand this controller is associated with (left or right)
+		uint32 PacketNum;		// The unique packet number the values in this struct relate to
+
+		float TouchPadXAnalog;	// X Axis values of the touchpad
+		float TouchPadYAnalog;	// Y Axis values of the touchpad
+
+		float TriggerAnalog;	// Float value of how much the trigger has been pressed
+
+		float HandGripAnalog;	// Float value of how much the grip has been pressed
+		float IndexGripAnalog;	// Float value of the index finger's curl
+		float MiddleGripAnalog;	// Float value of the middle finger's curl
+		float RingGripAnalog;	// Float value of the ring finger's curl
+		float PinkyGripAnalog;	// Float value of the pinky finger's curl
+		float ThumbGripAnalog;	// Float value of the thumb's curl
+
+		float ThumbIndexSplayAnalog;	// Float value of the thumb-index splay
+		float IndexMiddleSplayAnalog;	// Float value of the index-middle splay
+		float MiddleRingSplayAnalog;	// Float value of the middle-ring splay
+		float RingPinkySplayAnalog;		// Float value of the ring-pinky splay
+	};
+
+	/** Current input values of the active player controllers  */
+	FInputDeviceState ControllerStates[SteamVRInputDeviceConstants::MaxControllers];
+
+	/** Holds the action sets that will be handled by the SteamVR Input System  */
+	TArray<FSteamVRInputActionSet> SteamVRInputActionSets;
+
+	/** Holds the action sets that will be fed in to OpenVR  */
+	VRActiveActionSet_t ActiveActionSets[MAX_ACTION_SETS];
+
+	/** Holds the actions that will be handled by the SteamVR Input System  */
+	TArray<FSteamVRInputAction> Actions;
+
+	/** Holds the actions that will be handled by the SteamVR Input System  */
+	TArray<FSteamVRInputAction> ActionEvents;
+
+	/** A list of supported controller types that Controller Binding files will be generated for  */
+	TArray<FControllerType> ControllerTypes;
+
+	/** The handle to the action set that will be used in generating the Action Manifest and Controller Bindings  */
+	VRActionSetHandle_t MainActionSet;
+
+	/** The raw pose handle for the left controller  */
+	VRActionHandle_t VRControllerHandleLeft;
+
+	/** The raw pose handle for the right controller  */
+	VRActionHandle_t VRControllerHandleRight;
+
+	/** The tracker pose tagged as Special1 in UE  */
+	VRActionHandle_t VRSpecial1;
+
+	/** The tracker pose tagged as Special2 in UE  */
+	VRActionHandle_t VRSpecial2;
+
+	/** The tracker pose tagged as Special3 in UE  */
+	VRActionHandle_t VRSpecial3;
+
+	/** The tracker pose tagged as Special4 in UE  */
+	VRActionHandle_t VRSpecial4;
+
+	/** The tracker pose tagged as Special5 in UE  */
+	VRActionHandle_t VRSpecial5;
+
+	/** The tracker pose tagged as Special6 in UE  */
+	VRActionHandle_t VRSpecial6;
+
+	/** The tracker pose tagged as Special7 in UE  */
+	VRActionHandle_t VRSpecial7;
+
+	/** The tracker pose tagged as Special8 in UE  */
+	VRActionHandle_t VRSpecial8;
+
+	/** The handle for the skeletal input of the left hand  */
+	VRActionHandle_t VRSkeletalHandleLeft;
+
+	/** The handle for the skeletal input of the right hand  */
+	VRActionHandle_t VRSkeletalHandleRight;
+
+	/** The handle for the vibration of the right hand  */
+	VRActionHandle_t VRVibrationLeft;
+
+	/** The handle for the vibration of the right hand  */
+	VRActionHandle_t VRVibrationRight;
+
+	/** Current skeletal summary input values of the left hand  */
+	VRSkeletalSummaryData_t VRSkeletalSummaryDataLeft;
+
+	/** Current skeletal summary input values of the right hand  */
+	VRSkeletalSummaryData_t VRSkeletalSummaryDataRight;
+
+	/** The last input error that was encountered. Used for preventing spamming the output log  */
+	EVRInputError LastInputError = VRInputError_None;
+
+	/** Delay that will be applied for buttons  */
+	float InitialButtonRepeatDelay;
+
+	/** Delay that will be applied to detect button press repeats  */
+	float ButtonRepeatDelay;
+
+	/** Mapping of buttons and controllers  */
+	FGamepadKeyNames::Type Buttons[k_unMaxTrackedDeviceCount][ESteamVRInputButton::TotalButtonCount];
+
+	/** The filename for this project  */
+	FString GameFileName;
+
+	/** The name of the project as defined in the Project Settings' Title field */
+	FString GameProjectName;
+
+	/** The version of this project as defined in the Project Settings' Version field  */
+	FString GameProjectVersion;
+
+	/** The unique app key used by SteamVR to distinguish an Editor Session  */
+	FString EditorAppKey;
+
+	/** Whether or not skeletal input is supported and present in the controller in the player's left hand  */
+	bool bIsSkeletalControllerLeftPresent = false;
+
+	/** Whether or not skeletal input is supported and present in the controller in the player's right hand  */
+	bool bIsSkeletalControllerRightPresent = false;
+
+	/** Whether to use the Skeleton Pose for the Orientation and Position of the motion controller  */
+	bool bUseSkeletonPose = false;
+
+	/**
+	* The PredictedSecondsFromNow parameter for calls to Get Pose Action
+	* A value of -9999.f will trigger a GetPoseActionDataForNextFrame, otherwise a GetPoseActionRelativeToNow will be called
+	*/
+	float GlobalPredictedSecondsFromNow = -9999.f;
+
+	/** The skeletal tracking level for the controller in the player's left hand  */
+	EVRSkeletalTrackingLevel LeftControllerFidelity;
+
+	/** The skeletal tracking level for the controller in the player's right hand  */
+	EVRSkeletalTrackingLevel RightControllerFidelity;
+
+	/** Device identifier used to check for class validity */
+	int32 DeviceSignature = 2019;
+
+private:
+	/** 
+	 * Generate the controller bindings for SteamVR supported controllers
+	 * @param BindingsPath - Where the controller input profile json files will be generated (by default, this is Config\SteamVRBindings)
+	 * @param InOutControllerTypes - Common SteamVR controller types for which the bindings will be generated (e.g. knuckles)
+	 * @param InOutDefaultBindings - Metadata associated to each controller binding
+	 * @param InActionsArray - The list of SteamVR actions that needs to be generated for each controller
+	 * @param InInputMapping - The mapping of SteamVR actions and their controller inputs
+	 * @param bDeleteIfExists - Flag of whether or not to overwrite an existing controller binding
+	 */
+	void GenerateControllerBindings(const FString& BindingsPath, TArray<FControllerType>& InOutControllerTypes, TArray<TSharedPtr<FJsonValue>>& InOutDefaultBindings, TArray<FSteamVRInputAction>& InActionsArray, TArray<FInputMapping>& InInputMapping, bool bDeleteIfExists = false);
+
+	/**
+	* Generate the SteamVR specific action bindings that will generated for a controller
+	* @param InInputMapping - The mapping of SteamVR actions and their controller inputs
+	* @param JsonValuesArray - The generated mappings in json format
+	* @param bIsGenericController - Whether this is a Generic Controller (e.g. MotionController)
+	*/
+	void GenerateActionBindings(TArray<FInputMapping> &InInputMapping, TArray<TSharedPtr<FJsonValue>> &JsonValuesArray, FControllerType Controller, bool bIsGenericController = false);
+
+	/** Delegate called when an action mapping has been modified in the editor  */
+	FDelegateHandle ActionMappingsChangedHandle;
+
+	/** Provides a user-friendly version of the results of a SteamVR Input call in English to the Output Log */
+	void GetInputError(EVRInputError InputError, FString InputAction);
+
+	/** @deprecated Initialize the SteamVR device Ids to UE Controller Id mappings */
+	void InitControllerMappings();
+
+	/** Setup the keys used by supported SteamVR Controllers  */
+	void InitControllerKeys();
+
+	/** Process all input action events for a give action set */
+	void ProcessActionEvents(FSteamVRInputActionSet SteamVRInputActionSet);
+
+	/**
+	* Create the action manifest used by the SteamVR Input System
+	* @param GenerateActions - Whether to generate main actions from the project input settings
+	* @param GenerateBindings - Whether to generate the controller bindings for all supported SteamVR controllers
+	* @param RegisterApp - Whether to register currently running application as an Editor session
+	* @param DeleteBindings - Whether to overwrite current controller bindings (if any)
+	*/
+	void GenerateActionManifest(bool GenerateActions=true, bool GenerateBindings=true, bool RegisterApp=true, bool DeleteBindings=false);
+
+	/**
+	* Create the application manifest for an Editor session
+	* @param ManifestPath - Where the action manifest is located. By default this is under Config/SteamVRBindings
+	* @param ProjectName - The name of the project from the project settings' description field
+	* @param OutAppKey - The generated unique appkey for this application
+	* @param OutAppManifestPath - Where the application manifest will be generated to. Default is under the project's Config folder
+	* @return Whether or not the application manifest was generated successfully and the running Editor session has been properly registered to SteamVR
+	*/
+	bool GenerateAppManifest(FString ManifestPath, FString ProjectName, FString& OutAppKey, FString& OutAppManifestPath);
+
+	/**
+	* Utility function that builds a json object from a string of fields. Input strings should be paired (e.g fieldname1, fieldvalue1, fieldname2, fieldvalue2 ...)
+	* @param StringFields - The paired strings to generate a json object for (e.g fieldname1, fieldvalue1, fieldname2, fieldvalue2 ...)
+	* @param OutJsonObject - The resulting json object
+	* @return Whether or not we generated a valid json object
+	*/
+	bool BuildJsonObject(TArray<FString> StringFields, TSharedRef<FJsonObject> OutJsonObject);
+
+	/**
+	* Convert UE4 style key-based action bindings to SteamVR/OpenXR format
+	* @param InputSettings - The engine's input settings
+	* @param InOutUniqueInputs - The input mappings that needs to be processed. Will also hold the processed mappings
+	*/
+	void ProcessKeyInputMappings(const UInputSettings* InputSettings, TArray<FName> &InOutUniqueInputs);
+
+	/**
+	* Convert UE4 style axis-based action bindings to SteamVR/OpenXR format
+	* @param InputSettings - The engine's input settings
+	* @param InOutUniqueInputs - The input mappings that needs to be processed. Will also hold the processed mappings
+	*/
+	void ProcessKeyAxisMappings(const UInputSettings* InputSettings, TArray<FName> &InOutUniqueInputs);
+
+	/** Remove any invalid action to Axis mappings so they won#t be generated by the plugin */
+	void SanitizeActions();
+
+	/**
+	* Registers an Editor session to the SteamVR system
+	* @param ManifestPath - Path to the Application Manifest that will be generated. By default this will be under the Config folder
+	*/
+	void RegisterApplication(FString ManifestPath);
+
+	/** 
+	* Get the Skeletal Input Handle for a give controller input path
+	* @param ActionPath - The SteamVR user path for the controller hand we want the skeletal handle for
+	* @param SkeletalHandle - Will hold the skeletal handle if found. k_InvalidHandleValue if an error is encountered
+	* @return Whether or not Whether the Skeletal Handle was successfully read
+	*/	
+	bool SetSkeletalHandle(char* ActionPath, VRActionHandle_t& SkeletalHandle);
+
+	/**
+	* Look for any Key Mappings defined in the project's DefaultInput.ini 
+	* @param InputSettings - The engine's input settings
+	* @param ActionName - The name of the action where we want to find controller inputs for
+	* @param OutMappings - Will hold the Action to controller axis mapping result of the search
+	*/
+	void FindAxisMappings(const UInputSettings* InputSettings, const FName AxisName, TArray<FInputAxisKeyMapping>& OutMappings) const;
+
+	/**
+	* Utility funtion to add meta data to a UE Axis Key Mapping
+	* @param InUEKeyMappings - Array of UE Axis key Mappings to process
+	* @result OutMappings - Array of SteamVR Axis Key Mappings 
+	*/
+	void GetSteamVRMappings(TArray<FInputAxisKeyMapping>& InUEKeyMappings, TArray<FSteamVRAxisKeyMapping>& OutMappings);
+
+	/**
+	* Look for any Key Mappings defined in the project's DefaultInput.ini 
+	* @param InputSettings - The engine's input settings
+	* @param ActionName - The name of the action where we want to find controller inputs for
+	* @param OutMappings - Will hold the Action to controller key mapping result of the search
+	*/
+	void FindActionMappings(const UInputSettings* InputSettings, const FName ActionName, TArray<FInputActionKeyMapping>& OutMappings) const;
+
+	/**
+	* Send the axis value for the given controller input
+	* @param TrackedControllerRole - The tracking type/role for this controller
+	* @param AxisButton - The type of button the value is associated with
+	* @param AnalogValue - The input value (float) that is being sent
+	*/
+	void SendAnalogMessage(const ETrackedControllerRole TrackedControllerRole, const FGamepadKeyNames::Type AxisButton, float AnalogValue);
+
+	/**
+	* Remove any unsupported special characters on a provide string
+	* @param InOutString - The text to process
+	* @return A copy of the processed text
+	*/
+	FString SanitizeString(FString& InOutString);
+
+	/**
+	* Transform the bone transforms for one hand to fit the skeleton of the opposite hand
+	* @param BoneTransformsLS - The bone transforms to mirror, in local (parent) space
+	* @param BoneTransformCount - The number of elements in the BoneTransformsLS array.  Must equal the number of bones in the SteamVR skeleton
+	*/
+	void MirrorSteamVRSkeleton(VRBoneTransform_t* BoneTransformsLS, int32 BoneTransformCount) const;
+
+	/** Our Message handler to direct input from the SteamVRInput System to the game runtime */
+	TSharedRef<FGenericApplicationMessageHandler> MessageHandler;
+
+	/** Currently active SteamVRInput Error. Used to compare with the last error to avoid spamming the output log. */
+	EVRInputError SteamVRError = VRInputError_None;
+
+	/** Previous SteamVRInput Error. Used to compare against the current SteamVR Error to avoid spamming logs */
+	EVRInputError PrevSteamVRError = VRInputError_None;
+
+	/** Input Axis (Analog/Float) action mappings defined for this project. Available in DefaultInput.ini or via the Editor UI ProjectSettings > Engine > Input  */
+	TArray<FInputAxisKeyMapping> KeyAxisMappings;
+
+	/** Input Key (Digital/Boolean) action mappings defined for this project. Available in DefaultInput.ini or via the Editor UI ProjectSettings > Engine > Input  */
+	TArray<FInputActionKeyMapping> KeyMappings;
+
+	/** Input Axis (Analog/Float) action mappings defined for this project with extra metadata */
+	TArray<FSteamVRAxisKeyMapping> SteamVRKeyAxisMappings;
+
+	/** Input keys (Digital/Boolean) action mappings defined for this project with extra metadata */
+	TArray<FSteamVRInputKeyMapping> SteamVRKeyInputMappings;
+
+	/** Base Orientation defined by the XRTrackingSystem */
+	FQuat CachedBaseOrientation;
+
+	/** Base Position defined by the XRTrackingSystem */
+	FVector CachedBasePosition;
+
+	/**
+	*	Utility function to check if a given key is part of a VR Controller supported by SteamVR
+	*	@param InputKey - Name of the input key to test
+	*	@return bool - Whether the key is a VR Controller supported by SteamVR
+	*/
+	bool IsVRKey(FName InputKey);
+
+	/**
+	*	Utility function to check if a given key is part of a VR Controller supported by SteamVR
+	*	@param AxiskeyMapping - Access KeyMapping to test
+	*	@return bool - Whether the input is a Vector2D and has been processed successfully
+	*/
+	bool ProcessVector2D(FSteamVRAxisKeyMapping AxisKeyMapping);
+
+	/**
+	* Buffer for current delta time to get an accurate approximation of how long to play haptics for
+	*/
+	float CurrentDeltaTime;
+
+};
diff --git a/Engine/Plugins/Runtime/Steam/SteamVR/Source/SteamVRInputDevice/Private/SteamVRInputTypes.h b/Engine/Plugins/Runtime/Steam/SteamVR/Source/SteamVRInputDevice/Private/SteamVRInputTypes.h
new file mode 100644
--- /dev/null
+++ b/Engine/Plugins/Runtime/Steam/SteamVR/Source/SteamVRInputDevice/Private/SteamVRInputTypes.h
@@ -0,0 +1,496 @@
+/*
+Copyright 2019 Valve Corporation under https://opensource.org/licenses/BSD-3-Clause
+This code includes modifications by Epic Games.  Modifications (c) Epic Games, Inc.
+
+Redistribution and use in source and binary forms, with or without modification,
+are permitted provided that the following conditions are met:
+
+1. Redistributions of source code must retain the above copyright notice, this
+   list of conditions and the following disclaimer.
+
+2. Redistributions in binary form must reproduce the above copyright notice,
+   this list of conditions and the following disclaimer in the documentation
+   and/or other materials provided with the distribution.
+
+3. Neither the name of the copyright holder nor the names of its contributors
+   may be used to endorse or promote products derived from this software
+   without specific prior written permission.
+
+THIS SOFTWARE IS PROVIDED BY THE COPYRIGHT HOLDERS AND CONTRIBUTORS "AS IS" AND
+ANY EXPRESS OR IMPLIED WARRANTIES, INCLUDING, BUT NOT LIMITED TO, THE IMPLIED
+WARRANTIES OF MERCHANTABILITY AND FITNESS FOR A PARTICULAR PURPOSE ARE DISCLAIMED.
+IN NO EVENT SHALL THE COPYRIGHT HOLDER OR CONTRIBUTORS BE LIABLE FOR ANY DIRECT,
+INDIRECT, INCIDENTAL, SPECIAL, EXEMPLARY, OR CONSEQUENTIAL DAMAGES (INCLUDING,
+BUT NOT LIMITED TO, PROCUREMENT OF SUBSTITUTE GOODS OR SERVICES; LOSS OF USE, DATA,
+OR PROFITS; OR BUSINESS INTERRUPTION) HOWEVER CAUSED AND ON ANY THEORY OF LIABILITY,
+WHETHER IN CONTRACT, STRICT LIABILITY, OR TORT (INCLUDING NEGLIGENCE OR OTHERWISE)
+ARISING IN ANY WAY OUT OF THE USE OF THIS SOFTWARE, EVEN IF ADVISED OF THE
+POSSIBILITY OF SUCH DAMAGE.
+*/
+
+#pragma once
+
+#include "openvr.h"
+#include "GameFramework/InputSettings.h"
+#include "InputCoreTypes.h"
+
+using namespace vr;
+
+#define CONTROLLERS_PER_PLAYER			2
+#define GENERIC_TRACKER_PLAYER_NUM		0
+#define TOUCHPAD_AXIS					0
+#define TRIGGER_AXIS					1
+#define KNUCKLES_TOTAL_HAND_GRIP_AXIS	2
+#define KNUCKLES_UPPER_HAND_GRIP_AXIS	3
+#define KNUCKLES_LOWER_HAND_GRIP_AXIS	4
+#define STEAMVR_SKELETON_BONE_COUNT		31
+#define DOT_45DEG						0.707f
+#define TOUCHPAD_DEADZONE				0.0f
+
+// Manifest constants
+#define MAX_ACTION_SETS					25
+#define CONTROLLER_BINDING_PATH			"SteamVRBindings"
+#define ACTION_MANIFEST					"steamvr_manifest.json"
+#define ACTION_MANIFEST_UE				"steamvr_actions.json"
+#define APP_MANIFEST_FILE				"steamvr_ue_editor_app.json"
+#define APP_MANIFEST_PREFIX				"application.generated.ue."
+
+// Action paths
+#define ACTION_SET						"/actions/main"
+#define ACTION_PATH_IN					"/actions/main/in"
+#define ACTION_PATH_CONTROLLER_LEFT		"/actions/main/in/controllerleft"
+#define ACTION_PATH_CONTROLLER_RIGHT	"/actions/main/in/controllerright"
+#define ACTION_PATH_SPECIAL_BACK_L		"/actions/main/in/special1"
+#define ACTION_PATH_SPECIAL_BACK_R		"/actions/main/in/special2"
+#define ACTION_PATH_SPECIAL_FRONT_L		"/actions/main/in/special3"
+#define ACTION_PATH_SPECIAL_FRONT_R		"/actions/main/in/special4"
+#define ACTION_PATH_SPECIAL_FRONTR_L	"/actions/main/in/special5"
+#define ACTION_PATH_SPECIAL_FRONTR_R	"/actions/main/in/special6"
+#define ACTION_PATH_SPECIAL_PISTOL_L	"/actions/main/in/special7"
+#define ACTION_PATH_SPECIAL_PISTOL_R	"/actions/main/in/special8"
+#define ACTION_PATH_SKELETON_LEFT		"/actions/main/in/skeletonleft"		
+#define ACTION_PATH_SKELETON_RIGHT		"/actions/main/in/skeletonright"
+#define ACTION_PATH_SKEL_HAND_LEFT		"/skeleton/hand/left"
+#define ACTION_PATH_SKEL_HAND_RIGHT		"/skeleton/hand/right"
+#define ACTION_PATH_OPEN_CONSOLE		"/actions/main/in/open_console"
+#define ACTION_PATH_VIBRATE_LEFT		"/actions/main/out/vibrateleft"
+#define ACTION_PATH_VIBRATE_RIGHT		"/actions/main/out/vibrateright"
+
+// Input paths
+#define ACTION_PATH_HEAD_PROXIMITY		"/user/head/proximity"
+#define ACTION_PATH_CONT_RAW_LEFT		"/user/hand/left/pose/raw"
+#define ACTION_PATH_CONT_RAW_RIGHT		"/user/hand/right/pose/raw"
+#define ACTION_PATH_SPCL_BACK_LEFT		"/user/hand/left/pose/back"					// Special 1
+#define ACTION_PATH_SPCL_BACK_RIGHT		"/user/hand/right/pose/back"				// Special 2
+#define ACTION_PATH_SPCL_FRONT_LEFT		"/user/hand/left/pose/front"				// Special 3
+#define ACTION_PATH_SPCL_FRONT_RIGHT	"/user/hand/right/pose/front"				// Special 4
+#define ACTION_PATH_SPCL_FRONTR_LEFT	"/user/hand/left/pose/frontandrolled"		// Special 5
+#define ACTION_PATH_SPCL_FRONTR_RIGHT	"/user/hand/right/pose/frontandrolled"		// Special 6
+#define ACTION_PATH_SPCL_PISTOL_LEFT	"/user/hand/left/pose/pistolgrip"			// Special 7
+#define ACTION_PATH_SPCL_PISTOL_RIGHT	"/user/hand/right/pose/pistolgrip"			// Special 8
+#define ACTION_PATH_TRIGGER_LEFT		"/user/hand/left/input/trigger"
+#define ACTION_PATH_TRIGGER_RIGHT		"/user/hand/right/input/trigger"
+#define ACTION_PATH_BUMPER_LEFT			"/user/hand/left/input/bumper"
+#define ACTION_PATH_BUMPER_RIGHT		"/user/hand/right/input/bumper"
+#define ACTION_PATH_THUMBSTICK_LEFT		"/user/hand/left/input/thumbstick"
+#define ACTION_PATH_THUMBSTICK_RIGHT	"/user/hand/right/input/thumbstick"
+#define ACTION_PATH_TRACKPAD_LEFT		"/user/hand/left/input/trackpad"
+#define ACTION_PATH_TRACKPAD_RIGHT		"/user/hand/right/input/trackpad"
+#define ACTION_PATH_JOYSTICK_LEFT		"/user/hand/left/input/joystick"
+#define ACTION_PATH_JOYSTICK_RIGHT		"/user/hand/right/input/joystick"
+#define ACTION_PATH_GRIP_LEFT			"/user/hand/left/input/grip"
+#define ACTION_PATH_GRIP_RIGHT			"/user/hand/right/input/grip"
+#define ACTION_PATH_BTN_A_LEFT			"/user/hand/left/input/a"
+#define ACTION_PATH_BTN_A_RIGHT			"/user/hand/right/input/a"
+#define ACTION_PATH_BTN_B_LEFT			"/user/hand/left/input/b"
+#define ACTION_PATH_BTN_B_RIGHT			"/user/hand/right/input/b"
+#define ACTION_PATH_BTN_X_LEFT			"/user/hand/left/input/x"
+#define ACTION_PATH_BTN_Y_LEFT			"/user/hand/left/input/y"
+#define ACTION_PATH_USER_SKEL_LEFT		"/user/hand/left/input/skeleton/left"
+#define ACTION_PATH_USER_SKEL_RIGHT		"/user/hand/right/input/skeleton/right"
+#define ACTION_PATH_USER_VIB_LEFT		"/user/hand/left/output/haptic"
+#define ACTION_PATH_USER_VIB_RIGHT		"/user/hand/right/output/haptic"
+#define ACTION_PATH_PINCH_GRAB_LEFT		"/user/hand/left/input/pinch"
+#define ACTION_PATH_PINCH_GRAB_RIGHT	"/user/hand/right/input/pinch"
+#define ACTION_PATH_GRIP_GRAB_LEFT		"/user/hand/left/input/grip"
+#define ACTION_PATH_GRIP_GRAB_RIGHT		"/user/hand/right/input/grip"
+#define ACTION_PATH_APPMENU_LEFT        "/user/hand/left/input/application_menu"
+#define ACTION_PATH_APPMENU_RIGHT        "/user/hand/right/input/application_menu"
+
+
+namespace SteamVRInputDeviceConstants 
+{
+	/* The maximum number of Unreal controllers is 8, two motion controllers per Unreal controller */
+	static const int32 MaxUnrealControllers = 4;
+
+	/* Total number of supported tracked devices */
+	static const int32 MaxControllers = k_unMaxTrackedDeviceCount;
+
+	/* The maximum number of "Special" hand designations */
+	static const int32 MaxSpecialDesignations = ((int32)EControllerHand::Special_9 - (int32)EControllerHand::Special_1) + 1;
+}
+
+/* Buttons available on a SteamVR controller */
+struct ESteamVRInputButton
+{
+	enum Type
+	{
+		System = 0,
+		ApplicationMenu = 1,
+		Grip = 2,
+		DPadLeft = 3,
+		DPadUp = 4,
+		DPadRight = 5,
+		DPadDown = 6,
+		ButtonA = 7,
+
+		ProximitySensor = 31,
+
+		Axis0 = 32,
+		Axis1 = 33,
+		Axis2 = 34,
+		Axis3 = 35,
+		Axis4 = 36,
+
+		Touchpad = Axis0,
+		Trigger =Axis1,
+		DashboardBack =Grip,
+		KnucklesA =Grip,
+		KnucklesB =ApplicationMenu,
+		KnucklesJoyStick =Axis3,
+		TotalButtonCount = 64
+	};
+};
+
+/* Available SteamVR Input Action Types */
+enum EActionType
+{
+	Boolean,
+	Vector1,
+	Vector2,
+	Vector3,
+	Vibration,
+	Pose,
+	Skeleton,
+	Invalid,
+	UEVector
+};
+
+struct FSteamVRAxisKeyMapping 
+{
+	FInputAxisKeyMapping InputAxisKeyMapping;
+	bool bIsPartofVector2;
+	bool bIsPartofVector3;
+	FName XAxisName;
+	FName YAxisName;
+	FName ZAxisName;
+	FName XAxisKey;
+	FName YAxisKey;
+	FName ZAxisKey;
+	FString ActionName;
+	FString ActionNameWithPath;
+	FString ControllerName;
+
+	FSteamVRAxisKeyMapping(FInputAxisKeyMapping& inInputAxisKeyMapping,  bool inIsPartofVector2, bool inIsPartofVector3)
+		: InputAxisKeyMapping(inInputAxisKeyMapping)
+		, bIsPartofVector2(inIsPartofVector2)
+		, bIsPartofVector3(inIsPartofVector3)
+	{
+		XAxisName = NAME_None;
+		YAxisName = NAME_None;
+		ZAxisName = NAME_None;
+		XAxisKey = NAME_None;
+		YAxisKey = NAME_None;
+		ZAxisKey = NAME_None;
+		ActionName = "";
+		ControllerName = "";
+	}
+};
+
+struct FSteamVRInputKeyMapping
+{
+	FInputActionKeyMapping InputKeyMapping;
+	FString ActionName;
+	FString ActionNameWithPath;
+	FString ControllerName;
+
+	FSteamVRInputKeyMapping(FInputActionKeyMapping& inInputKeyMapping)
+		: InputKeyMapping(inInputKeyMapping)
+	{
+		ActionName = "";
+		ControllerName = "";
+	}
+};
+
+const FString SActionTypes[] = {
+			TEXT("boolean"),
+			TEXT("vector1"),
+			TEXT("vector2"),
+			TEXT("vector3"),
+			TEXT("vibration"),
+			TEXT("pose"),
+			TEXT("skeleton"),
+			TEXT("")
+};
+
+struct FControllerType
+{
+	bool	bIsGenerated;
+	FName	Name;
+	FString	Description;
+	FString KeyEquivalent;
+	FName	ControllerName;
+	bool bIsActive;
+
+	FControllerType() {}
+	FControllerType(const FName& inName, const FString& inDescription, const FString& inKeyEquivalent)
+		: bIsGenerated(false)
+		, Name(inName)
+		, Description(inDescription)
+		, KeyEquivalent(inKeyEquivalent)
+	{
+		bIsActive = false;
+	}
+};
+
+struct FActionPath
+{
+	FString	Path;
+
+	FActionPath() {}
+	FActionPath(const FString& inPath)
+		: Path(inPath)
+	{}
+};
+
+struct FActionSource
+{
+	FName			Mode;
+	FString			Path;
+
+	FActionSource() {}
+	FActionSource(const FName& inMode, const FString& inPath)
+		: Mode(inMode),
+		Path(inPath)
+	{}
+};
+
+struct FInputMapping
+{
+	FName	InputKey;			// The UE Key (e.g.Motion_Controller_Thumbstick_X)
+	TArray<FString> Actions;	// Collection of SteamVR actions (e.g. Teleport, OpenConsole)
+
+	FInputMapping() {}
+};
+
+struct FSteamVRInputActionSet
+{
+	int32		Priority;					// The priority of this action set relative to other action sets.	
+	FName		Name;						// The name of the action set
+	FString		RestrictedToDevicePath;		// Device path that this action set is active for
+	FString		SecondaryActionSetPath;		// The path for the secondary action set
+
+	VRActionSetHandle_t Handle;						// This action set's handle
+	VRInputValueHandle_t RestrictedToDeviceHandle;	// Handle of a device path that this action set should be active for.  Use k_ulInvalidInputValueHandle to activate for all devices.
+	VRActionSetHandle_t SecondaryActionSetHandle;	// Secondary action set handle, if RestrictedToDeviceHandle is k_ulInvalidInputValueHandle, this is ignored
+
+	FSteamVRInputActionSet()
+	{}
+
+	FSteamVRInputActionSet(int32 InPriority, FName InName, VRActionSetHandle_t InHandle)
+		: Priority (InPriority)
+		, Name (InName)
+		, Handle (InHandle)
+	{
+		RestrictedToDeviceHandle = k_ulInvalidInputValueHandle;	// All devices
+		SecondaryActionSetHandle = k_ulInvalidActionSetHandle;	// Placeholder. Ignored due to k_ulInvalidValueHandle in RestrictedToDeviceHandle
+	}
+};
+
+struct FSteamVRInputAction
+{
+	FString		Path;							// The path defined for the action (e.g. main/in/{ActionName})
+	FName		Name;							// The SteamVR name of the action (e.g. Teleport, OpenConsole)
+	EActionType	Type;							// The SteamVR input data type
+	FName		KeyX;							// The UE Key in the X axis or float axis (e.g. Motion_Controller_Thumbstick_X)
+	FName		KeyY;							// The UE Key in the Y axis
+	FName		KeyZ;							// The UE Key in the Z axis
+	FVector		Value;							// The 1D, 2D, 3D (analog) input value of this action
+	FString		StringPath;						// The string value of this action (such as for Skeleton Paths where a bool or float axis is not appropriate)
+	bool		bState;							// The bool (digital) value of this action
+	bool		bRequirement;					// Whether or not this action is required
+
+	VRActionHandle_t Handle;					// The handle to the SteamVR main Action Set
+	VRInputValueHandle_t ActiveOrigin = 0;		// The input value handle of the origin of the latest input event
+	EVRInputError LastError;					// A cache for the last Error for operations against this action (could also be "No Error")
+
+	FString GetActionTypeName()
+	{
+		return SActionTypes[(int)Type];
+	}
+
+	FSteamVRInputAction(const FString& inPath, EActionType inType, bool inRequirement, const FName& inName, const FString& inStringPath)
+		: Path(inPath)
+		, Name(inName)
+		, Type(inType)
+		, KeyX()
+		, KeyY()
+		, KeyZ()
+		, Value()
+		, StringPath(inStringPath)
+		, bRequirement(inRequirement)
+		, Handle()
+		, LastError(VRInputError_None)
+	{}
+
+	FSteamVRInputAction(const FString& inPath, const FName& inName, const FName& inKeyName, bool inState)
+		: Path(inPath)
+		, Name(inName)
+		, Type(Boolean)
+		, KeyX(inKeyName)
+		, KeyY()
+		, KeyZ()
+		, Value()
+		, bState(inState)
+		, bRequirement(true)
+		, Handle()
+		, LastError(VRInputError_None)
+	{}
+
+	FSteamVRInputAction(const FString& inPath, const FName& inName, bool inRequirement, const FName& inKeyName, bool inState)
+		: Path(inPath)
+		, Name(inName)
+		, Type(Boolean)
+		, KeyX(inKeyName)
+		, KeyY()
+		, KeyZ()
+		, Value()
+		, bState(inState)
+		, bRequirement(inRequirement)
+		, Handle()
+		, LastError(VRInputError_None)
+	{}
+
+	FSteamVRInputAction(const FString& inPath, const FName& inName, const FName& inKeyName, float inValue1D)
+		: Path(inPath)
+		, Name(inName)
+		, Type(Vector1)
+		, KeyX(inKeyName)
+		, KeyY()
+		, KeyZ()
+		, Value(inValue1D, 0.f, 0.f)
+		, bRequirement(true)
+		, Handle()
+		, LastError(VRInputError_None)
+	{}
+
+	FSteamVRInputAction(const FString& inPath, const FName& inName, const FName& inKeyName_X, const FName& inKeyName_Y, const FVector2D& inValue2D)
+		: Path(inPath)
+		, Name(inName)
+		, Type(Vector2)
+		, KeyX(inKeyName_X)
+		, KeyY(inKeyName_Y)
+		, KeyZ()
+		, Value(inValue2D.X, inValue2D.Y, 0.f)
+		, bRequirement(true)
+		, Handle()
+		, LastError(VRInputError_None)
+	{}
+
+	FSteamVRInputAction(const FString& inPath, const FName& inName, const FName& inKeyName_X, const FName& inKeyName_Y, const FName& inKeyName_Z, const FVector& inValue3D)
+		: Path(inPath)
+		, Name(inName)
+		, Type(Vector3)
+		, KeyX(inKeyName_X)
+		, KeyY(inKeyName_Y)
+		, KeyZ(inKeyName_Z)
+		, Value(inValue3D.X, inValue3D.Y, inValue3D.Z)
+		, bRequirement(true)
+		, Handle()
+		, LastError(VRInputError_None)
+	{}
+
+	FSteamVRInputAction(const FString& inPath, const EActionType& inActionType, const bool& inRequirement, const FName& inName)
+		: Path(inPath)
+		, Name(inName)
+		, Type(inActionType)
+		, KeyX()
+		, KeyY()
+		, KeyZ()
+		, Value()
+		, bRequirement(inRequirement)
+		, Handle()
+		, LastError(VRInputError_None)
+	{}
+
+};
+
+struct FSteamVRInputState
+{
+	bool bIsAxis;
+	bool bIsAxis2;
+	bool bIsAxis3;
+	bool bIsTrigger;
+	bool bIsBumper;
+	bool bIsThumbstick;
+	bool bIsJoystick;
+	bool bIsTrackpad;
+	bool bIsDpadUp;
+	bool bIsDpadDown;
+	bool bIsDpadLeft;
+	bool bIsDpadRight;
+	bool bIsGrip;
+	bool bIsCapSense;
+	bool bIsLeft;
+	bool bIsFaceButton1;
+	bool bIsFaceButton2;
+	bool bIsXButton;
+	bool bIsYButton;
+	bool bIsGripGrab;
+	bool bIsPinchGrab;
+	bool bIsPress;
+	bool bIsAppMenu;
+	bool bIsProximity;
+
+	FSteamVRInputState() {}
+};
+
+struct FSteamVRTemporaryAction
+{
+	FKey UE4Key;
+	FName ActionName;
+	FKey UE4KeyX;
+	FKey UE4KeyY;
+	FKey UE4KeyVector;
+	bool bIsY;
+
+	FSteamVRTemporaryAction()
+	{
+		ActionName = NAME_None;
+		UE4Key = EKeys::Invalid;
+		UE4KeyX = EKeys::Invalid;
+		UE4KeyY = EKeys::Invalid;
+		UE4KeyVector = EKeys::Invalid;
+		bIsY = false;
+	}
+
+	FSteamVRTemporaryAction(const FKey& inUE4Key, const FName& inActionName)
+		: UE4Key(inUE4Key)
+		, ActionName(inActionName)
+	{
+		UE4KeyX = EKeys::Invalid;
+		UE4KeyY = EKeys::Invalid;
+		UE4KeyVector = EKeys::Invalid;
+		bIsY = false;
+	}
+
+	FSteamVRTemporaryAction(const FKey& inUE4Key, const FName& inActionName, const FKey& inUE4KeyX, const FKey& inUE4KeyY)
+		: UE4Key(inUE4Key)
+		, ActionName(inActionName)
+		, UE4KeyX(inUE4KeyX)
+		, UE4KeyY(inUE4KeyY)
+	{
+		UE4KeyVector = EKeys::Invalid;
+		bIsY = false;
+	}
+};
diff --git a/Engine/Plugins/Runtime/Steam/SteamVR/Source/SteamVRInputDevice/Private/SteamVRKnucklesKeys.h b/Engine/Plugins/Runtime/Steam/SteamVR/Source/SteamVRInputDevice/Private/SteamVRKnucklesKeys.h
new file mode 100644
--- /dev/null
+++ b/Engine/Plugins/Runtime/Steam/SteamVR/Source/SteamVRInputDevice/Private/SteamVRKnucklesKeys.h
@@ -0,0 +1,121 @@
+/*
+Copyright 2019 Valve Corporation under https://opensource.org/licenses/BSD-3-Clause
+This code includes modifications by Epic Games.  Modifications (c) Epic Games, Inc.
+
+Redistribution and use in source and binary forms, with or without modification,
+are permitted provided that the following conditions are met:
+
+1. Redistributions of source code must retain the above copyright notice, this
+   list of conditions and the following disclaimer.
+
+2. Redistributions in binary form must reproduce the above copyright notice,
+   this list of conditions and the following disclaimer in the documentation
+   and/or other materials provided with the distribution.
+
+3. Neither the name of the copyright holder nor the names of its contributors
+   may be used to endorse or promote products derived from this software
+   without specific prior written permission.
+
+THIS SOFTWARE IS PROVIDED BY THE COPYRIGHT HOLDERS AND CONTRIBUTORS "AS IS" AND
+ANY EXPRESS OR IMPLIED WARRANTIES, INCLUDING, BUT NOT LIMITED TO, THE IMPLIED
+WARRANTIES OF MERCHANTABILITY AND FITNESS FOR A PARTICULAR PURPOSE ARE DISCLAIMED.
+IN NO EVENT SHALL THE COPYRIGHT HOLDER OR CONTRIBUTORS BE LIABLE FOR ANY DIRECT,
+INDIRECT, INCIDENTAL, SPECIAL, EXEMPLARY, OR CONSEQUENTIAL DAMAGES (INCLUDING,
+BUT NOT LIMITED TO, PROCUREMENT OF SUBSTITUTE GOODS OR SERVICES; LOSS OF USE, DATA,
+OR PROFITS; OR BUSINESS INTERRUPTION) HOWEVER CAUSED AND ON ANY THEORY OF LIABILITY,
+WHETHER IN CONTRACT, STRICT LIABILITY, OR TORT (INCLUDING NEGLIGENCE OR OTHERWISE)
+ARISING IN ANY WAY OUT OF THE USE OF THIS SOFTWARE, EVEN IF ADVISED OF THE
+POSSIBILITY OF SUCH DAMAGE.
+*/
+
+#pragma once
+#include "InputCoreTypes.h"
+
+namespace KnucklesControllerKeyNames
+{
+	// Knuckles Finger Curls
+	const FGamepadKeyNames::Type SteamVR_Knuckles_Left_Finger_Index_Curl("Knuckles_Left_Finger_Index_Curl");
+	const FGamepadKeyNames::Type SteamVR_Knuckles_Right_Finger_Index_Curl("Knuckles_Right_Finger_Index_Curl");
+
+	const FGamepadKeyNames::Type SteamVR_Knuckles_Left_Finger_Middle_Curl("Knuckles_Left_Finger_Middle_Curl");
+	const FGamepadKeyNames::Type SteamVR_Knuckles_Right_Finger_Middle_Curl("Knuckles_Right_Finger_Middle_Curl");
+
+	const FGamepadKeyNames::Type SteamVR_Knuckles_Left_Finger_Ring_Curl("Knuckles_Left_Finger_Ring_Curl");
+	const FGamepadKeyNames::Type SteamVR_Knuckles_Right_Finger_Ring_Curl("Knuckles_Right_Finger_Ring_Curl");
+
+	const FGamepadKeyNames::Type SteamVR_Knuckles_Left_Finger_Pinky_Curl("Knuckles_Left_Finger_Pinky_Curl");
+	const FGamepadKeyNames::Type SteamVR_Knuckles_Right_Finger_Pinky_Curl("Knuckles_Right_Finger_Pinky_Curl");
+
+	const FGamepadKeyNames::Type SteamVR_Knuckles_Left_Finger_Thumb_Curl("Knuckles_Left_Finger_Thumb_Curl");
+	const FGamepadKeyNames::Type SteamVR_Knuckles_Right_Finger_Thumb_Curl("Knuckles_Right_Finger_Thumb_Curl");
+
+	// Knuckles Finger Splays
+	const FGamepadKeyNames::Type SteamVR_Knuckles_Left_Finger_ThumbIndex_Splay("Knuckles_Left_Finger_ThumbIndex_Splay");
+	const FGamepadKeyNames::Type SteamVR_Knuckles_Right_Finger_ThumbIndex_Splay("Knuckles_Right_Finger_ThumbIndex_Splay");
+
+	const FGamepadKeyNames::Type SteamVR_Knuckles_Left_Finger_IndexMiddle_Splay("Knuckles_Left_Finger_IndexMiddle_Splay");
+	const FGamepadKeyNames::Type SteamVR_Knuckles_Right_Finger_IndexMiddle_Splay("Knuckles_Right_Finger_IndexMiddle_Splay");
+
+	const FGamepadKeyNames::Type SteamVR_Knuckles_Left_Finger_MiddleRing_Splay("Knuckles_Left_Finger_MiddleRing_Splay");
+	const FGamepadKeyNames::Type SteamVR_Knuckles_Right_Finger_MiddleRing_Splay("Knuckles_Right_Finger_MiddleRing_Splay");
+
+	const FGamepadKeyNames::Type SteamVR_Knuckles_Left_Finger_RingPinky_Splay("Knuckles_Left_Finger_RingPinky_Splay");
+	const FGamepadKeyNames::Type SteamVR_Knuckles_Right_Finger_RingPinky_Splay("Knuckles_Right_Finger_RingPinky_Splay");
+
+}
+
+namespace KnucklesControllerKeys
+{
+	// Knuckles CapSense
+	const FKey SteamVR_Knuckles_Left_A_CapSense("Knuckles_Left_A_CapSense");
+	const FKey SteamVR_Knuckles_Right_A_CapSense("Knuckles_Right_A_CapSense");
+	const FKey SteamVR_Knuckles_Left_B_CapSense("Knuckles_Left_B_CapSense");
+	const FKey SteamVR_Knuckles_Right_B_CapSense("Knuckles_Right_B_CapSense");
+
+	const FKey SteamVR_Knuckles_Left_Trigger_CapSense("Knuckles_Left_Trigger_CapSense");
+	const FKey SteamVR_Knuckles_Right_Trigger_CapSense("Knuckles_Right_Trigger_CapSense");
+
+	const FKey SteamVR_Knuckles_Left_Thumbstick_CapSense("Knuckles_Left_Thumbstick_CapSense");
+	const FKey SteamVR_Knuckles_Right_Thumbstick_CapSense("Knuckles_Right_Thumbstick_CapSense");
+
+	const FKey SteamVR_Knuckles_Left_Trackpad_CapSense("Knuckles_Left_Trackpad_CapSense");
+	const FKey SteamVR_Knuckles_Right_Trackpad_CapSense("Knuckles_Right_Trackpad_CapSense");
+
+	const FKey SteamVR_Knuckles_Left_Trackpad_GripForce("Knuckles_Left_Trackpad_GripForce");
+	const FKey SteamVR_Knuckles_Right_Trackpad_GripForce("Knuckles_Right_Trackpad_GripForce");
+
+	// Knuckles Trackpad
+	const FKey SteamVR_Knuckles_Left_Trackpad_X("Knuckles_Left_Trackpad_X");
+	const FKey SteamVR_Knuckles_Right_Trackpad_X("Knuckles_Right_Trackpad_X");
+	const FKey SteamVR_Knuckles_Left_Trackpad_Y("Knuckles_Left_Trackpad_Y");
+	const FKey SteamVR_Knuckles_Right_Trackpad_Y("Knuckles_Right_Trackpad_Y");
+
+	// Knuckles Curls
+	const FKey SteamVR_Knuckles_Left_Finger_Index_Curl("Knuckles_Left_Finger_Index_Curl");
+	const FKey SteamVR_Knuckles_Right_Finger_Index_Curl("Knuckles_Right_Finger_Index_Curl");
+
+	const FKey SteamVR_Knuckles_Left_Finger_Middle_Curl("Knuckles_Left_Finger_Middle_Curl");
+	const FKey SteamVR_Knuckles_Right_Finger_Middle_Curl("Knuckles_Right_Finger_Middle_Curl");
+
+	const FKey SteamVR_Knuckles_Left_Finger_Ring_Curl("Knuckles_Left_Finger_Ring_Curl");
+	const FKey SteamVR_Knuckles_Right_Finger_Ring_Curl("Knuckles_Right_Finger_Ring_Curl");
+
+	const FKey SteamVR_Knuckles_Left_Finger_Pinky_Curl("Knuckles_Left_Finger_Pinky_Curl");
+	const FKey SteamVR_Knuckles_Right_Finger_Pinky_Curl("Knuckles_Right_Finger_Pinky_Curl");
+
+	const FKey SteamVR_Knuckles_Left_Finger_Thumb_Curl("Knuckles_Left_Finger_Thumb_Curl");
+	const FKey SteamVR_Knuckles_Right_Finger_Thumb_Curl("Knuckles_Right_Finger_Thumb_Curl");
+
+	// Knuckles Splays
+	const FKey SteamVR_Knuckles_Left_Finger_ThumbIndex_Splay("Knuckles_Left_Finger_ThumbIndex_Splay");
+	const FKey SteamVR_Knuckles_Right_Finger_ThumbIndex_Splay("Knuckles_Right_Finger_ThumbIndex_Splay");
+
+	const FKey SteamVR_Knuckles_Left_Finger_IndexMiddle_Splay("Knuckles_Left_Finger_IndexMiddle_Splay");
+	const FKey SteamVR_Knuckles_Right_Finger_IndexMiddle_Splay("Knuckles_Right_Finger_IndexMiddle_Splay");
+
+	const FKey SteamVR_Knuckles_Left_Finger_MiddleRing_Splay("Knuckles_Left_Finger_MiddleRing_Splay");
+	const FKey SteamVR_Knuckles_Right_Finger_MiddleRing_Splay("Knuckles_Right_Finger_MiddleRing_Splay");
+
+	const FKey SteamVR_Knuckles_Left_Finger_RingPinky_Splay("Knuckles_Left_Finger_RingPinky_Splay");
+	const FKey SteamVR_Knuckles_Right_Finger_RingPinky_Splay("Knuckles_Right_Finger_RingPinky_Splay");
+}
diff --git a/Engine/Plugins/Runtime/Steam/SteamVR/Source/SteamVRInputDevice/Private/UE4HandSkeletonDefinition.h b/Engine/Plugins/Runtime/Steam/SteamVR/Source/SteamVRInputDevice/Private/UE4HandSkeletonDefinition.h
new file mode 100644
--- /dev/null
+++ b/Engine/Plugins/Runtime/Steam/SteamVR/Source/SteamVRInputDevice/Private/UE4HandSkeletonDefinition.h
@@ -0,0 +1,80 @@
+/*
+Copyright 2019 Valve Corporation under https://opensource.org/licenses/BSD-3-Clause
+This code includes modifications by Epic Games.  Modifications (c) Epic Games, Inc.
+
+Redistribution and use in source and binary forms, with or without modification,
+are permitted provided that the following conditions are met:
+
+1. Redistributions of source code must retain the above copyright notice, this
+   list of conditions and the following disclaimer.
+
+2. Redistributions in binary form must reproduce the above copyright notice,
+   this list of conditions and the following disclaimer in the documentation
+   and/or other materials provided with the distribution.
+
+3. Neither the name of the copyright holder nor the names of its contributors
+   may be used to endorse or promote products derived from this software
+   without specific prior written permission.
+
+THIS SOFTWARE IS PROVIDED BY THE COPYRIGHT HOLDERS AND CONTRIBUTORS "AS IS" AND
+ANY EXPRESS OR IMPLIED WARRANTIES, INCLUDING, BUT NOT LIMITED TO, THE IMPLIED
+WARRANTIES OF MERCHANTABILITY AND FITNESS FOR A PARTICULAR PURPOSE ARE DISCLAIMED.
+IN NO EVENT SHALL THE COPYRIGHT HOLDER OR CONTRIBUTORS BE LIABLE FOR ANY DIRECT,
+INDIRECT, INCIDENTAL, SPECIAL, EXEMPLARY, OR CONSEQUENTIAL DAMAGES (INCLUDING,
+BUT NOT LIMITED TO, PROCUREMENT OF SUBSTITUTE GOODS OR SERVICES; LOSS OF USE, DATA,
+OR PROFITS; OR BUSINESS INTERRUPTION) HOWEVER CAUSED AND ON ANY THEORY OF LIABILITY,
+WHETHER IN CONTRACT, STRICT LIABILITY, OR TORT (INCLUDING NEGLIGENCE OR OTHERWISE)
+ARISING IN ANY WAY OUT OF THE USE OF THIS SOFTWARE, EVEN IF ADVISED OF THE
+POSSIBILITY OF SUCH DAMAGE.
+*/
+
+#pragma once
+
+#include "CoreMinimal.h"
+
+/**
+ * Enum for each bone in the UE4 hand skeleton
+*/
+enum EUE4HandBone : int8
+{
+	EUE4HandBone_Wrist = 0,
+	EUE4HandBone_Index_01,
+	EUE4HandBone_Index_02,
+	EUE4HandBone_Index_03,
+	EUE4HandBone_Middle_01,
+	EUE4HandBone_Middle_02,
+	EUE4HandBone_Middle_03,
+	EUE4HandBone_Pinky_01,
+	EUE4HandBone_Pinky_02,
+	EUE4HandBone_Pinky_03,
+	EUE4HandBone_Ring_01,
+	EUE4HandBone_Ring_02,
+	EUE4HandBone_Ring_03,
+	EUE4HandBone_Thumb_01,
+	EUE4HandBone_Thumb_02,
+	EUE4HandBone_Thumb_03,
+	EUE4HandBone_Count
+};
+
+
+/**
+ * Convenience functions to access static data as defined 
+ * by the UE4 VR reference hand skeleton
+*/
+namespace UE4HandSkeleton
+{
+	/** Returns the number of bones in the skeleton */
+	inline int32	GetBoneCount() { return EUE4HandBone_Count; }
+
+	/** Returns the name of the bone at the given index */
+	const FName&	GetBoneName(int32 nBoneIndex);
+
+	/** Returns the index of the parent bone of the given bone.  Returns -1 if the bone does not have a parent */
+	int32			GetParentIndex(int32 nBoneIndex);
+
+	/** Returns the number of children of the given bone */
+	int32			GetChildCount(int32 nBoneIndex);
+
+	/** Returns the index of the nth child of the given bone */
+	int32			GetChildIndex(int32 nBoneIndex, int32 nChildIndex);
+};
diff --git a/Engine/Plugins/Runtime/Steam/SteamVR/Source/SteamVRInputDevice/Public/SteamVRControllerKeys.h b/Engine/Plugins/Runtime/Steam/SteamVR/Source/SteamVRInputDevice/Public/SteamVRControllerKeys.h
deleted file mode 100644
diff --git a/Engine/Plugins/Runtime/Steam/SteamVR/Source/SteamVRInputDevice/Public/SteamVRInputDevice.h b/Engine/Plugins/Runtime/Steam/SteamVR/Source/SteamVRInputDevice/Public/SteamVRInputDevice.h
deleted file mode 100644
diff --git a/Engine/Plugins/Runtime/Steam/SteamVR/Source/SteamVRInputDevice/Public/SteamVRInputDeviceFunctionLibrary.h b/Engine/Plugins/Runtime/Steam/SteamVR/Source/SteamVRInputDevice/Public/SteamVRInputDeviceFunctionLibrary.h
--- a/Engine/Plugins/Runtime/Steam/SteamVR/Source/SteamVRInputDevice/Public/SteamVRInputDeviceFunctionLibrary.h
+++ b/Engine/Plugins/Runtime/Steam/SteamVR/Source/SteamVRInputDevice/Public/SteamVRInputDeviceFunctionLibrary.h
@@ -30,8 +30,10 @@ POSSIBILITY OF SUCH DAMAGE.
 
 #pragma once
 
+#include "openvr.h"
+#include "GameFramework/InputSettings.h"
+#include "InputCoreTypes.h"
 #include "CoreMinimal.h"
-#include "SteamVRInputDevice.h"
 #include "Kismet/BlueprintFunctionLibrary.h"
 #include "Features/IModularFeatures.h"
 #include "Runtime/Core/Public/GenericPlatform/GenericPlatformProcess.h"
@@ -40,6 +42,8 @@ POSSIBILITY OF SUCH DAMAGE.
 #include "Kismet/GameplayStatics.h"
 #include "SteamVRInputDeviceFunctionLibrary.generated.h"
 
+using namespace vr;
+
 #define ACTION_PATH_VIBRATE_LEFT		"/actions/main/out/vibrateleft"
 #define ACTION_PATH_VIBRATE_RIGHT		"/actions/main/out/vibrateright"
 #define MAX_BINDINGINFO_COUNT	5
@@ -385,9 +389,6 @@ class STEAMVRINPUTDEVICE_API USteamVRInputDeviceFunctionLibrary : public UBluepr
 	GENERATED_BODY()
 
 public:
-	/** Retrieve the first available SteamVR Input device currently active in a game */
-	static FSteamVRInputDevice* GetSteamVRInputDevice();
-
 	/**
 	* Get the finger curl and splay for a give hand in the current frame
 	* @param Hand - Which hand to get the finger curls and splay values for
@@ -636,4 +637,8 @@ public:
 	
 	/** Open the SteamVR Controller Input Dashboard in the user#s default browser */
 	static void LaunchBindingsURL();
+
+private:
+	/** Retrieve the first available SteamVR Input device currently active in a game */
+	static class FSteamVRInputDevice* GetSteamVRInputDevice();
 };
diff --git a/Engine/Plugins/Runtime/Steam/SteamVR/Source/SteamVRInputDevice/Public/SteamVRInputTypes.h b/Engine/Plugins/Runtime/Steam/SteamVR/Source/SteamVRInputDevice/Public/SteamVRInputTypes.h
deleted file mode 100644
diff --git a/Engine/Plugins/Runtime/Steam/SteamVR/Source/SteamVRInputDevice/Public/SteamVRKnucklesKeys.h b/Engine/Plugins/Runtime/Steam/SteamVR/Source/SteamVRInputDevice/Public/SteamVRKnucklesKeys.h
deleted file mode 100644
diff --git a/Engine/Plugins/Runtime/Steam/SteamVR/Source/SteamVRInputDevice/Public/UE4HandSkeletonDefinition.h b/Engine/Plugins/Runtime/Steam/SteamVR/Source/SteamVRInputDevice/Public/UE4HandSkeletonDefinition.h
deleted file mode 100644
