Fix non-unity in light of CL# 10903717
#rnx

diff --git a/Engine/Plugins/Runtime/Steam/SteamVR/Source/SteamVRInputDevice/Private/SteamVRInputDeviceFunctionLibrary.cpp b/Engine/Plugins/Runtime/Steam/SteamVR/Source/SteamVRInputDevice/Private/SteamVRInputDeviceFunctionLibrary.cpp
--- a/Engine/Plugins/Runtime/Steam/SteamVR/Source/SteamVRInputDevice/Private/SteamVRInputDeviceFunctionLibrary.cpp
+++ b/Engine/Plugins/Runtime/Steam/SteamVR/Source/SteamVRInputDevice/Private/SteamVRInputDeviceFunctionLibrary.cpp
@@ -34,6 +34,7 @@ POSSIBILITY OF SUCH DAMAGE.
 #include "openvr.h"
 #include "GameFramework/WorldSettings.h"
 #include "HAL/FileManagerGeneric.h"
+#include "SteamVRInputDevice.h"
 using namespace vr;
 //#endif // STEAMVRCONTROLLER_SUPPORTED_PLATFORMS
 
