Merge SteamVR Input hotfix changes from Valve:

(1) Assert crash/check when enabling debug options for SteamVR 1.9.7 and above
(2) Fix issue when an action is inadvertently triggered (can happen when a specific controller adds a mapping to multiple actions where another controller isn't mapped to exactly the same actions)
(3) Skip bIsGenerated checks for other valid controllers when one is found.
(4) Cleanly shutdown when Marketplace or custom GitHub plugin build is activw in the project
(5) Trigger axis values only when there's a change from the previous value
(6) Fix Cosmos axis events showing up as bindable in blueprint
(7) Issue # 100: Trackpad Touch vs Click for Index Controllers
(8) Issue # 102: Thumbstick d-pad bindings default to "sub_mode": click
(9) Issue #91: Noisy logs
(10) Change Request from Epic: Add repeat key events for the VREditor

diff --git a/Engine/Plugins/Runtime/Steam/SteamVR/Source/SteamVRInputDevice/Private/SteamVRInputDevice.cpp b/Engine/Plugins/Runtime/Steam/SteamVR/Source/SteamVRInputDevice/Private/SteamVRInputDevice.cpp
--- a/Engine/Plugins/Runtime/Steam/SteamVR/Source/SteamVRInputDevice/Private/SteamVRInputDevice.cpp
+++ b/Engine/Plugins/Runtime/Steam/SteamVR/Source/SteamVRInputDevice/Private/SteamVRInputDevice.cpp
@@ -1635,14 +1635,6 @@ void FSteamVRInputDevice::GenerateActionBindings(TArray<FInputMapping> &InInputM
 			// Let's check if there're any SteamVR specific key that already exists for this action
 			for (FSteamVRInputKeyMapping SteamVRKeyInputMappingInner : SteamVRKeyInputMappings)
 			{
-				// Ensure the hmd proximity key is always generated
-				//if (SteamVRKeyInputMappingInner.InputKeyMapping.Key.GetFName().ToString().Contains(TEXT("SteamVR_HMD_Proximity")))
-				//{
-				//	UE_LOG(LogTemp, Warning, TEXT("*** Action includes a SteamVR_HMD_Proximity!"));
-				//	bIsGenericController = false;
-				//	break;
-				//}
-
 				// Check for generic controllers that have steamvr inputs already defined
 				if (SteamVRKeyInputMapping.InputKeyMapping.ActionName.ToString().Equals(SteamVRKeyInputMappingInner.InputKeyMapping.ActionName.ToString())
 					&& IsVRKey(SteamVRKeyInputMappingInner.InputKeyMapping.Key.GetFName())
@@ -1667,10 +1659,6 @@ void FSteamVRInputDevice::GenerateActionBindings(TArray<FInputMapping> &InInputM
 				&& !SteamVRKeyInputMapping.InputKeyMapping.Key.GetFName().ToString().Contains(TEXT("HMD_Proximity"))
 				)
 			{
-				//UE_LOG(LogTemp, Warning, TEXT("*** %s %s %s"), *Controller.KeyEquivalent, 
-				//	*SteamVRKeyInputMapping.InputKeyMapping.Key.GetFName().ToString(),
-				//	*SteamVRKeyInputMapping.ControllerName
-				//	);
 				continue;
 			}
 			else
@@ -1734,16 +1722,8 @@ void FSteamVRInputDevice::GenerateActionBindings(TArray<FInputMapping> &InInputM
 				}
 				else
 				{
-					// Set cap sense input state
-					if (CurrentInputKeyName.Contains(TEXT("ValveIndex")) && InputState.bIsTrackpad)
-					{
-						InputState.bIsCapSense = true;					
-					}
-					else
-					{
-						InputState.bIsCapSense = CurrentInputKeyName.Contains(TEXT("CapSense"), ESearchCase::CaseSensitive, ESearchDir::FromEnd) ||
-							CurrentInputKeyName.Contains(TEXT("_Touch"), ESearchCase::CaseSensitive, ESearchDir::FromEnd);
-					}
+					InputState.bIsCapSense = CurrentInputKeyName.Contains(TEXT("CapSense"), ESearchCase::CaseSensitive, ESearchDir::FromEnd) ||
+						CurrentInputKeyName.Contains(TEXT("_Touch"), ESearchCase::CaseSensitive, ESearchDir::FromEnd);
 				}
 
 				// Check for DPad Keys
@@ -1814,7 +1794,6 @@ void FSteamVRInputDevice::GenerateActionBindings(TArray<FInputMapping> &InInputM
 				CacheMode = InputState.bIsPress ? FName(TEXT("button")) : CacheMode;
 				CacheMode = InputState.bIsTrackpad ? FName(TEXT("trackpad")) : CacheMode;
 				CacheMode = InputState.bIsThumbstick ? FName(TEXT("joystick")) : CacheMode;
-				//CacheMode = (InputState.bIsTrackpad || InputState.bIsThumbstick) && !InputState.bIsAxis ? FName(TEXT("button")) : CacheMode;
 				CacheMode = InputState.bIsGrip ? FName(TEXT("button")) : CacheMode;
 				CacheMode = InputState.bIsPinchGrab || InputState.bIsGripGrab ? FName(TEXT("grab")) : CacheMode;
 
@@ -1827,10 +1806,6 @@ void FSteamVRInputDevice::GenerateActionBindings(TArray<FInputMapping> &InInputM
 				{
 					CachePath = InputState.bIsLeft ? FString(TEXT(ACTION_PATH_BUMPER_LEFT)) : FString(TEXT(ACTION_PATH_BUMPER_RIGHT));
 				}
-				//else if (InputState.bIsThumbstick)
-				//{
-				//	CachePath = InputState.bIsLeft ? FString(TEXT(ACTION_PATH_THUMBSTICK_LEFT)) : FString(TEXT(ACTION_PATH_THUMBSTICK_RIGHT));
-				//}
 				else if (InputState.bIsTrackpad)
 				{
 					CachePath = InputState.bIsLeft ? FString(TEXT(ACTION_PATH_TRACKPAD_LEFT)) : FString(TEXT(ACTION_PATH_TRACKPAD_RIGHT));
@@ -1910,9 +1885,16 @@ void FSteamVRInputDevice::GenerateActionBindings(TArray<FInputMapping> &InInputM
 				// Add parameters if Dpad
 				if (InputState.bIsDpadUp || InputState.bIsDpadDown || InputState.bIsDpadLeft || InputState.bIsDpadRight)
 				{
-					// Create Submode
+					// Create Submode - Vive will use click as default
 					TSharedRef<FJsonObject> SubmodeJsonObject = MakeShareable(new FJsonObject());
-					SubmodeJsonObject->SetStringField(TEXT("sub_mode"), TEXT("click"));
+					if (Controller.Description.Contains(TEXT("Vive")))
+					{
+						SubmodeJsonObject->SetStringField(TEXT("sub_mode"), TEXT("click"));
+					}
+					else
+					{
+						SubmodeJsonObject->SetStringField(TEXT("sub_mode"), TEXT("touch"));
+					}
 					
 					// Create Parameter
 					TSharedPtr<FJsonObject> ParametersJsonObject = MakeShareable(new FJsonObject());
@@ -3489,6 +3471,7 @@ void FSteamVRInputDevice::RegisterApplication(FString ManifestPath)
 		GetInputError(InputError, FString(TEXT("Setting main action set")));
 
 		// Add to action set array
+		SteamVRInputActionSets.Empty();
 		SteamVRInputActionSets.Add(FSteamVRInputActionSet(0, ACTION_SET, MainActionSet));
 
 		// Populate Active Action sets that will later be used in OpenVR calls
@@ -3602,35 +3585,55 @@ void FSteamVRInputDevice::ProcessActionEvents(FSteamVRInputActionSet SteamVRInpu
 			// Get digital data from SteamVR
 			InputDigitalActionData_t DigitalData;
 			EVRInputError ActionStateError = VRInput()->GetDigitalActionData(Action.Handle, &DigitalData, sizeof(DigitalData), k_ulInvalidInputValueHandle);
-
+			
 			if (ActionStateError != VRInputError_None)
 			{
 				//GetInputError(ActionStateError, TEXT("Error encountered retrieving digital data from SteamVR"));
 				//UE_LOG(LogSteamVRInputDevice, Error, TEXT("%s"), *Action.Path);
 				continue;
 			}
-			else if (ActionStateError == VRInputError_None &&
-				DigitalData.bActive &&
-				DigitalData.bState != Action.bState
-				)
+			else if (ActionStateError == VRInputError_None && DigitalData.bActive)
 			{
-				// Update our action to the value reported by SteamVR
-				Action.bState = DigitalData.bState;
+				// Send event back to Engine
+				if (Action.KeyX != NAME_None)
+				{
+					// Update the active origin
+					Action.ActiveOrigin = DigitalData.activeOrigin;
 
-				// Update the active origin
-				Action.ActiveOrigin = DigitalData.activeOrigin;
+					if (DigitalData.bState)
+					{
+						if (!Action.bState)
+						{
+							MessageHandler->OnControllerButtonPressed(Action.KeyX, 0, false);
+							Action.bState = DigitalData.bState;
+							Action.LastUpdated = DigitalData.fUpdateTime;
+							Action.bIsRepeat = false;
+							//UE_LOG(LogTemp, Warning, TEXT("Handle %s KeyX %s Value %i Changed %i UpdateTime %f"), *Action.Path, *Action.KeyX.ToString(), DigitalData.bState, DigitalData.bChanged, DigitalData.fUpdateTime);
+						}	
+						else
+						{
+							float EffectiveDelay = Action.bIsRepeat ? REPEAT_DIGITAL_ACTION_DELAY : INITIAL_DIGITAL_ACTION_DELAY;
 
-				// Test what we're receiving from SteamVR
-				//UE_LOG(LogTemp, Warning, TEXT("Handle %s KeyX %s Value %i"), *Action.Path, *FoundKey.GetFName().ToString(), DigitalData.bState);
+							if (Action.LastUpdated - DigitalData.fUpdateTime >= EffectiveDelay)
+							{
+								MessageHandler->OnControllerButtonPressed(Action.KeyX, 0, true);
+								Action.LastUpdated = DigitalData.fUpdateTime;
+								Action.bIsRepeat = true;
+								//UE_LOG(LogTemp, Warning, TEXT("[REPEAT] Handle %s KeyX %s Value %i Changed %i UpdateTime %f"), *Action.Path, *Action.KeyX.ToString(), DigitalData.bState, DigitalData.bChanged, DigitalData.fUpdateTime);
+							}							
+						}
+					}
+					else 
+					{
+						if (Action.bState)
+						{
+							MessageHandler->OnControllerButtonReleased(Action.KeyX, 0, false);
+							//UE_LOG(LogTemp, Display, TEXT("Handle %s KeyX %s Value %i Changed %i UpdateTime %f"), *Action.Path, *Action.KeyX.ToString(), DigitalData.bState, DigitalData.bChanged, DigitalData.fUpdateTime);
+						}
 
-				// Sent event back to Engine
-				if (Action.bState && Action.KeyX != NAME_None)
-				{
-					MessageHandler->OnControllerButtonPressed(Action.KeyX, 0, false);
-				}
-				else
-				{
-					MessageHandler->OnControllerButtonReleased(Action.KeyX, 0, false);
+						Action.bState = DigitalData.bState;
+						Action.bIsRepeat = false;
+					}
 				}
 			}
 		}
@@ -3652,7 +3655,7 @@ void FSteamVRInputDevice::ProcessActionEvents(FSteamVRInputActionSet SteamVRInpu
 				// Update the active origin
 				Action.ActiveOrigin = AnalogData.activeOrigin;
 
-				if (Action.KeyX != NAME_None && (AnalogData.deltaX > KINDA_SMALL_NUMBER || Action.Name.IsEqual(FName(CONTROLLER_BINDING_PATH))))
+				if (Action.KeyX != NAME_None && (FMath::Abs(AnalogData.deltaX) > KINDA_SMALL_NUMBER || Action.Name.IsEqual(FName(CONTROLLER_BINDING_PATH))))
 				{
 					// Test what we're receiving from SteamVR
 					//UE_LOG(LogTemp, Warning, TEXT("Handle %s KeyX %s X-Value [%f]"), *Action.Path, *Action.KeyX.ToString(), AnalogData.x);
@@ -3661,7 +3664,7 @@ void FSteamVRInputDevice::ProcessActionEvents(FSteamVRInputActionSet SteamVRInpu
 					MessageHandler->OnControllerAnalog(Action.KeyX, 0, Action.Value.X);
 				}
 
-				if (Action.KeyY != NAME_None && (AnalogData.deltaY > KINDA_SMALL_NUMBER))
+				if (Action.KeyY != NAME_None && (FMath::Abs(AnalogData.deltaY) > KINDA_SMALL_NUMBER))
 				{
 					// Test what we're receiving from SteamVR
 					//UE_LOG(LogTemp, Warning, TEXT("Handle %s KeyY %s Y-Value {%f}"), *Action.Path, *Action.KeyY.ToString(), AnalogData.y);
diff --git a/Engine/Plugins/Runtime/Steam/SteamVR/Source/SteamVRInputDevice/Private/SteamVRInputTypes.h b/Engine/Plugins/Runtime/Steam/SteamVR/Source/SteamVRInputDevice/Private/SteamVRInputTypes.h
--- a/Engine/Plugins/Runtime/Steam/SteamVR/Source/SteamVRInputDevice/Private/SteamVRInputTypes.h
+++ b/Engine/Plugins/Runtime/Steam/SteamVR/Source/SteamVRInputDevice/Private/SteamVRInputTypes.h
@@ -46,6 +46,8 @@ using namespace vr;
 #define STEAMVR_SKELETON_BONE_COUNT		31
 #define DOT_45DEG						0.707f
 #define TOUCHPAD_DEADZONE				0.0f
+#define INITIAL_DIGITAL_ACTION_DELAY	0.2f	
+#define REPEAT_DIGITAL_ACTION_DELAY		0.1f	
 
 // Manifest constants
 #define MAX_ACTION_SETS					25
@@ -324,6 +326,9 @@ struct FSteamVRInputAction
 	VRInputValueHandle_t ActiveOrigin = 0;		// The input value handle of the origin of the latest input event
 	EVRInputError LastError;					// A cache for the last Error for operations against this action (could also be "No Error")
 
+	bool		bIsRepeat;						// If this action is being repeated
+	float		LastUpdated;					// Last update time as reported by SteamVR
+
 	FString GetActionTypeName()
 	{
 		return SActionTypes[(int)Type];
@@ -341,7 +346,10 @@ struct FSteamVRInputAction
 		, bRequirement(inRequirement)
 		, Handle()
 		, LastError(VRInputError_None)
-	{}
+	{
+		bIsRepeat = false;
+		LastUpdated = 0.f;
+	}
 
 	FSteamVRInputAction(const FString& inPath, const FName& inName, const FName& inKeyName, bool inState)
 		: Path(inPath)
@@ -355,7 +363,10 @@ struct FSteamVRInputAction
 		, bRequirement(true)
 		, Handle()
 		, LastError(VRInputError_None)
-	{}
+	{
+		bIsRepeat = false;
+		LastUpdated = 0.f;
+	}
 
 	FSteamVRInputAction(const FString& inPath, const FName& inName, bool inRequirement, const FName& inKeyName, bool inState)
 		: Path(inPath)
@@ -369,7 +380,10 @@ struct FSteamVRInputAction
 		, bRequirement(inRequirement)
 		, Handle()
 		, LastError(VRInputError_None)
-	{}
+	{
+		bIsRepeat = false;
+		LastUpdated = 0.f;
+	}
 
 	FSteamVRInputAction(const FString& inPath, const FName& inName, const FName& inKeyName, float inValue1D)
 		: Path(inPath)
@@ -395,7 +409,10 @@ struct FSteamVRInputAction
 		, bRequirement(true)
 		, Handle()
 		, LastError(VRInputError_None)
-	{}
+	{
+		bIsRepeat = false;
+		LastUpdated = 0.f;
+	}
 
 	FSteamVRInputAction(const FString& inPath, const FName& inName, const FName& inKeyName_X, const FName& inKeyName_Y, const FName& inKeyName_Z, const FVector& inValue3D)
 		: Path(inPath)
@@ -408,7 +425,10 @@ struct FSteamVRInputAction
 		, bRequirement(true)
 		, Handle()
 		, LastError(VRInputError_None)
-	{}
+	{
+		bIsRepeat = false;
+		LastUpdated = 0.f;
+	}
 
 	FSteamVRInputAction(const FString& inPath, const EActionType& inActionType, const bool& inRequirement, const FName& inName)
 		: Path(inPath)
@@ -421,7 +441,10 @@ struct FSteamVRInputAction
 		, bRequirement(inRequirement)
 		, Handle()
 		, LastError(VRInputError_None)
-	{}
+	{
+		bIsRepeat = false;
+		LastUpdated = 0.f;
+	}
 
 };
 
