PR #6278: [BUGFIX] STEAMVR Fix crash (wrong usage) of SteamVR StereoLayer API. (Contributed by [REDACTED_PERSON])

diff --git a/Engine/Plugins/Runtime/Steam/SteamVR/Source/SteamVR/Private/SteamVRStereoLayers.cpp b/Engine/Plugins/Runtime/Steam/SteamVR/Source/SteamVR/Private/SteamVRStereoLayers.cpp
--- a/Engine/Plugins/Runtime/Steam/SteamVR/Source/SteamVR/Private/SteamVRStereoLayers.cpp
+++ b/Engine/Plugins/Runtime/Steam/SteamVR/Source/SteamVR/Private/SteamVRStereoLayers.cpp
@@ -17,6 +17,10 @@ static vr::EVROverlayError sOvrError;
 #define OVR_VERIFY(x) x
 #endif
 
+#if !PLATFORM_MAC
+#include "VulkanRHIPrivate.h"
+#endif
+
 /*=============================================================================
 *
 * Helper functions
@@ -104,8 +108,18 @@ void FSteamVRHMD::UpdateLayer(struct FSteamVRLayer& Layer, uint32 LayerId, bool
 	vr::VRTextureBounds_t TextureBounds;
     TextureBounds.uMin = Layer.LayerDesc.UVRect.Min.X;
     TextureBounds.uMax = Layer.LayerDesc.UVRect.Max.X;
-    TextureBounds.vMin = Layer.LayerDesc.UVRect.Min.Y;
-    TextureBounds.vMax = Layer.LayerDesc.UVRect.Max.Y;
+
+	if ( IsOpenGLPlatform( GMaxRHIShaderPlatform ) )
+	{
+		TextureBounds.vMin = Layer.LayerDesc.UVRect.Max.Y;
+		TextureBounds.vMax = Layer.LayerDesc.UVRect.Min.Y;
+	}
+	else
+	{
+		TextureBounds.vMin = Layer.LayerDesc.UVRect.Min.Y;
+		TextureBounds.vMax = Layer.LayerDesc.UVRect.Max.Y;
+	}
+
 	OVR_VERIFY(VROverlay->SetOverlayTextureBounds(Layer.OverlayHandle, &TextureBounds));
 	const float WorldToMeterScale = GetWorldToMetersScale();
 	check(WorldToMeterScale > 0.f);
@@ -211,18 +225,39 @@ void FSteamVRHMD::UpdateStereoLayers_RenderThread()
 			if (Layer.bUpdateTexture || (Layer.LayerDesc.Flags & LAYER_FLAG_TEX_CONTINUOUS_UPDATE))
 			{
 				vr::Texture_t Texture;
-				Texture.handle = Layer.LayerDesc.Texture->GetNativeResource();
+				vr::VRVulkanTextureData_t VulkanTexture{};
 				if ( IsVulkanPlatform( GMaxRHIShaderPlatform ) )
 				{
+					auto vlkRHI = static_cast<FVulkanDynamicRHI*>(GDynamicRHI);
+					FRHITexture2D* TextureRHI2D = Layer.LayerDesc.Texture->GetTexture2D();
+					check(TextureRHI2D);
+					FVulkanTexture2D* Texture2D = (FVulkanTexture2D*)TextureRHI2D;
+
+					VulkanTexture.m_pInstance = vlkRHI->GetInstance();
+					VulkanTexture.m_pDevice = vlkRHI->GetDevice()->GetInstanceHandle();
+					VulkanTexture.m_pPhysicalDevice = vlkRHI->GetDevice()->GetPhysicalHandle();
+					VulkanTexture.m_pQueue = vlkRHI->GetDevice()->GetGraphicsQueue()->GetHandle();
+					VulkanTexture.m_nQueueFamilyIndex = vlkRHI->GetDevice()->GetGraphicsQueue()->GetFamilyIndex();
+					VulkanTexture.m_nImage = (uint64_t)Texture2D->Surface.Image;
+					VulkanTexture.m_nWidth = Texture2D->Surface.Width;
+					VulkanTexture.m_nHeight = Texture2D->Surface.Height;
+					VulkanTexture.m_nFormat = (uint32_t)Texture2D->Surface.ViewFormat;
+					VulkanTexture.m_nSampleCount = 1;
+
+					Texture.handle = &VulkanTexture;
 					Texture.eType = vr::TextureType_Vulkan;
 				}
 				else if ( IsOpenGLPlatform( GMaxRHIShaderPlatform ) )
 				{
+					// We need to dereference the pointer to the real handle
+					uint32 TextureID = *reinterpret_cast<uint32*>(Layer.LayerDesc.Texture->GetNativeResource());
+					Texture.handle = reinterpret_cast<void*>(TextureID);
 					Texture.eType = vr::TextureType_OpenGL;
 				}
 				else
 				{
 #if PLATFORM_WINDOWS
+					Texture.handle = Layer.LayerDesc.Texture->GetNativeResource();
 					Texture.eType = vr::TextureType_DirectX;
 #else
 					check( 0 );
