Fix compile errors on Mac after CL 10946291

diff --git a/Engine/Plugins/Runtime/Steam/SteamVR/Source/SteamVR/Private/SteamVRStereoLayers.cpp b/Engine/Plugins/Runtime/Steam/SteamVR/Source/SteamVR/Private/SteamVRStereoLayers.cpp
--- a/Engine/Plugins/Runtime/Steam/SteamVR/Source/SteamVR/Private/SteamVRStereoLayers.cpp
+++ b/Engine/Plugins/Runtime/Steam/SteamVR/Source/SteamVR/Private/SteamVRStereoLayers.cpp
@@ -228,6 +228,9 @@ void FSteamVRHMD::UpdateStereoLayers_RenderThread()
 				vr::VRVulkanTextureData_t VulkanTexture{};
 				if ( IsVulkanPlatform( GMaxRHIShaderPlatform ) )
 				{
+#if PLATFORM_MAC
+					check( 0 );
+#else
 					auto vlkRHI = static_cast<FVulkanDynamicRHI*>(GDynamicRHI);
 					FRHITexture2D* TextureRHI2D = Layer.LayerDesc.Texture->GetTexture2D();
 					check(TextureRHI2D);
@@ -246,12 +249,13 @@ void FSteamVRHMD::UpdateStereoLayers_RenderThread()
 
 					Texture.handle = &VulkanTexture;
 					Texture.eType = vr::TextureType_Vulkan;
+#endif
 				}
 				else if ( IsOpenGLPlatform( GMaxRHIShaderPlatform ) )
 				{
 					// We need to dereference the pointer to the real handle
 					uint32 TextureID = *reinterpret_cast<uint32*>(Layer.LayerDesc.Texture->GetNativeResource());
-					Texture.handle = reinterpret_cast<void*>(TextureID);
+					Texture.handle = static_cast<void*>(TextureID);
 					Texture.eType = vr::TextureType_OpenGL;
 				}
 				else
