Fix incorrect pointer case after CL 10949850

diff --git a/Engine/Plugins/Runtime/Steam/SteamVR/Source/SteamVR/Private/SteamVRStereoLayers.cpp b/Engine/Plugins/Runtime/Steam/SteamVR/Source/SteamVR/Private/SteamVRStereoLayers.cpp
--- a/Engine/Plugins/Runtime/Steam/SteamVR/Source/SteamVR/Private/SteamVRStereoLayers.cpp
+++ b/Engine/Plugins/Runtime/Steam/SteamVR/Source/SteamVR/Private/SteamVRStereoLayers.cpp
@@ -254,8 +254,8 @@ void FSteamVRHMD::UpdateStereoLayers_RenderThread()
 				else if ( IsOpenGLPlatform( GMaxRHIShaderPlatform ) )
 				{
 					// We need to dereference the pointer to the real handle
-					uint32 TextureID = *reinterpret_cast<uint32*>(Layer.LayerDesc.Texture->GetNativeResource());
-					Texture.handle = static_cast<void*>(TextureID);
+					uintptr_t TextureID = *reinterpret_cast<uint32*>(Layer.LayerDesc.Texture->GetNativeResource());
+					Texture.handle = reinterpret_cast<void*>(TextureID);
 					Texture.eType = vr::TextureType_OpenGL;
 				}
 				else
