PR #6582: Last set of changes for SteamVRInput in UE4.24 Hotfix 2 (CL11079797) (Contributed by [REDACTED_PERSON])

(1) Remove odd VRInit during Initialization
(2) Ensure all relevant SteamVRInput code runs only when SteamVRHMD is active
(3) Ensure that we proeprly re ∩┐╜ reg app, reg manifests, etc during SteamVR runtime restarts
(4) Regression from new feature axis mappings to only trigger when there's a significant change in axis data: Fix a use case when Get Finger Curls and Splays swallow axis inputs after this new feature has been introduced.

diff --git a/Engine/Plugins/Runtime/Steam/SteamVR/Source/SteamVRInputDevice/Private/SteamVRInputDevice.cpp b/Engine/Plugins/Runtime/Steam/SteamVR/Source/SteamVRInputDevice/Private/SteamVRInputDevice.cpp
--- a/Engine/Plugins/Runtime/Steam/SteamVR/Source/SteamVRInputDevice/Private/SteamVRInputDevice.cpp
+++ b/Engine/Plugins/Runtime/Steam/SteamVR/Source/SteamVRInputDevice/Private/SteamVRInputDevice.cpp
@@ -44,6 +44,7 @@ POSSIBILITY OF SUCH DAMAGE.
 #include "Runtime/HeadMountedDisplay/Public/IXRTrackingSystem.h"
 #include "SteamVRSkeletonDefinition.h"
 #include "Misc/MessageDialog.h"
+#include "ISteamVRPlugin.h"
 
 #if PLATFORM_WINDOWS
 #include "Windows/WindowsHWrapper.h"
@@ -123,18 +124,9 @@ FSteamVRInputDevice::~FSteamVRInputDevice()
 void FSteamVRInputDevice::InitSteamVRSystem()
 {
 	//UE_LOG(LogTemp, Warning, TEXT("Attempting to load Steam VR System..."));
-	if (!VRSystem())
-	{
-		// Initialize OpenVR
-		EVRInitError SteamVRInitError = VRInitError_Driver_NotLoaded;
-		IVRSystem* SteamVRSystem = vr::VR_Init(&SteamVRInitError, vr::VRApplication_Scene);
+	SteamVRHMDModule = FModuleManager::LoadModulePtr<ISteamVRPlugin>(TEXT("SteamVR"));
 
-		if (SteamVRInitError != VRInitError_None)
-		{
-			return;
-		}
-	}
-	else if (VRSystem() && VRInput() && IsInGameThread())
+	if (SteamVRHMDModule && SteamVRHMDModule->GetVRSystem() && VRSystem() && VRInput())
 	{
 		UE_LOG(LogSteamVRInputDevice, Display, TEXT("SteamVR runtime %u.%u.%u loaded."), k_nSteamVRVersionMajor, k_nSteamVRVersionMinor, k_nSteamVRVersionBuild);
 		
@@ -143,7 +135,7 @@ void FSteamVRInputDevice::InitSteamVRSystem()
 		bIsSkeletalControllerRightPresent = SetSkeletalHandle(TCHAR_TO_UTF8(*FString(TEXT(ACTION_PATH_SKELETON_RIGHT))), VRSkeletalHandleRight);
 
 		// (Re)Load Action Manifest
-		GenerateActionManifest(false, false, true, false);
+		GenerateActionManifest(true, true, true, false);
 
 		// Set haptic handles
 		LastInputError = VRInput()->GetActionHandle(TCHAR_TO_UTF8(*FString(TEXT(ACTION_PATH_VIBRATE_LEFT))), &VRVibrationLeft);
@@ -157,6 +149,8 @@ void FSteamVRInputDevice::InitSteamVRSystem()
 		{
 			VRVibrationRight = k_ulInvalidActionHandle;
 		}
+
+		bSteamVRWasShutdown = false;
 	}
 }
 
@@ -166,23 +160,25 @@ void FSteamVRInputDevice::Tick(float DeltaTime)
 	CurrentDeltaTime = DeltaTime;
 
 	// Watch for SteamVR availability & restarts
-	if (!VRSystem())
+	if ((SteamVRHMDModule && !SteamVRHMDModule->GetVRSystem()) || bSteamVRWasShutdown)
 	{
+		bSteamVRWasShutdown = true;
 		InitSteamVRSystem();
+		//UE_LOG(LogSteamVRInputDevice, Warning, TEXT("SteamVR System Inactive. Trying to initialize..."));
 	}
-
-	// Cache the controller transform to ensure ResetOrientationAndPosition gets the correct values (Valid for UE4.18 upwards)
-	// [REDACTED_PROFILE_URL]/steamvr_unreal_plugin/issues/2
-	if (GEngine->XRSystem.IsValid())
+	else if(SteamVRHMDModule && SteamVRHMDModule->GetVRSystem())
 	{
-		CachedBaseOrientation = GEngine->XRSystem->GetBaseOrientation();
-		CachedBasePosition = GEngine->XRSystem->GetBasePosition();
-	}
-	else
-	{
-		CachedBaseOrientation = FQuat::Identity;
-		CachedBasePosition = FVector::ZeroVector;
+		// Cache the controller transform to ensure ResetOrientationAndPosition gets the correct values (Valid for UE4.18 upwards)
+		// [REDACTED_PROFILE_URL]/steamvr_unreal_plugin/issues/2
+		if (GEngine->XRSystem.IsValid() && SteamVRHMDModule && SteamVRHMDModule->GetVRSystem())
+		{
+			CachedBaseOrientation = GEngine->XRSystem->GetBaseOrientation();
+			CachedBasePosition = GEngine->XRSystem->GetBasePosition();
+		}
 	}
+
+	CachedBaseOrientation = FQuat::Identity;
+	CachedBasePosition = FVector::ZeroVector;
 }
 
 void FSteamVRInputDevice::FindAxisMappings(const UInputSettings* InputSettings, const FName InAxisName, TArray<FInputAxisKeyMapping>& OutMappings) const
@@ -357,8 +353,7 @@ void FSteamVRInputDevice::SendAnalogMessage(const ETrackedControllerRole Tracked
 
 void FSteamVRInputDevice::SendControllerEvents()
 {
-
-	if (VRSystem() && VRInput() && SteamVRInputActionSets.Num() > 0)
+	if (SteamVRHMDModule && SteamVRHMDModule->GetVRSystem() && VRSystem() && VRInput() && SteamVRInputActionSets.Num() > 0)
 	{
 		EVRInputError ActionStateError = VRInput()->UpdateActionState(ActiveActionSets, sizeof(VRActiveActionSet_t), 1);
 
@@ -388,7 +383,7 @@ bool FSteamVRInputDevice::Exec(UWorld* InWorld, const TCHAR* Cmd, FOutputDevice&
 
 bool FSteamVRInputDevice::GetControllerOrientationAndPosition(const int32 ControllerIndex, const EControllerHand DeviceHand, FRotator& OutOrientation, FVector& OutPosition, float WorldToMetersScale) const
 {
-	if (VRInput() && VRCompositor())
+	if (SteamVRHMDModule && SteamVRHMDModule->GetVRSystem() && VRInput() && VRCompositor())
 	{
 		InputPoseActionData_t PoseData = {};
 		EVRInputError InputError = VRInputError_NoData;
@@ -643,17 +638,18 @@ bool FSteamVRInputDevice::GetControllerOrientationAndPosition(const int32 Contro
 			OrientationQuat.Normalize();
 			OutOrientation = OrientationQuat.Rotator();
 
+			return true;
 		}
 	}
 
-	return true;
+	return false;
 }
 
 ETrackingStatus FSteamVRInputDevice::GetControllerTrackingStatus(const int32 ControllerIndex, const EControllerHand DeviceHand) const
 {
 	ETrackingStatus TrackingStatus = ETrackingStatus::NotTracked;
 
-	if (VRInput() && VRCompositor())
+	if (SteamVRHMDModule && SteamVRHMDModule->GetVRSystem() && VRInput() && VRCompositor())
 	{
 		InputPoseActionData_t PoseData = {};
 		EVRInputError InputError = VRInputError_NoData;
@@ -1363,7 +1359,11 @@ void FSteamVRInputDevice::ReloadActionManifest()
 	if (VRSystem() && VRInput() && VRApplications())
 	{
 		// Set Action Manifest Path
+#if WITH_EDITOR
 		const FString ManifestPath = FPaths::ProjectConfigDir() / CONTROLLER_BINDING_PATH / ACTION_MANIFEST;
+#else
+		const FString ManifestPath = FPaths::EngineConfigDir() / CONTROLLER_BINDING_PATH / ACTION_MANIFEST;
+#endif
 		UE_LOG(LogSteamVRInputDevice, Display, TEXT("Reloading Action Manifest in: %s"), *ManifestPath);
 			
 		// Load application manifest
@@ -3459,7 +3459,7 @@ void FSteamVRInputDevice::RegisterApplication(FString ManifestPath)
 		#if WITH_EDITOR
 			TheActionManifestPath = IFileManager::Get().ConvertToAbsolutePathForExternalAppForRead(*ManifestPath);
 		#else
-			TheActionManifestPath = FPaths::ConvertRelativePathToFull(FPaths::ProjectDir() / TEXT("Config") / TEXT("SteamVRBindings") / TEXT(ACTION_MANIFEST)).Replace(TEXT("/"), TEXT("\\"));
+			TheActionManifestPath = FPaths::ConvertRelativePathToFull(FPaths::EngineDir() / TEXT("Config") / TEXT("SteamVRBindings") / TEXT(ACTION_MANIFEST)).Replace(TEXT("/"), TEXT("\\"));
 		#endif
 		
 		UE_LOG(LogSteamVRInputDevice, Display, TEXT("[STEAMVR INPUT] Trying to load Action Manifest from: %s"), *TheActionManifestPath);
diff --git a/Engine/Plugins/Runtime/Steam/SteamVR/Source/SteamVRInputDevice/Private/SteamVRInputDevice.h b/Engine/Plugins/Runtime/Steam/SteamVR/Source/SteamVRInputDevice/Private/SteamVRInputDevice.h
--- a/Engine/Plugins/Runtime/Steam/SteamVR/Source/SteamVRInputDevice/Private/SteamVRInputDevice.h
+++ b/Engine/Plugins/Runtime/Steam/SteamVR/Source/SteamVRInputDevice/Private/SteamVRInputDevice.h
@@ -42,6 +42,7 @@ POSSIBILITY OF SUCH DAMAGE.
 #include "Serialization/JsonSerializer.h"
 #include "SteamVRInputTypes.h"
 #include "SteamVRInputPublic.h"
+#include "ISteamVRPlugin.h"
 
 class STEAMVRINPUTDEVICE_API FSteamVRInputDevice : public IInputDevice, public FXRMotionControllerBase, public IHapticDevice
 {
@@ -76,6 +77,9 @@ public:
 	virtual float GetHapticAmplitudeScale() const override;
 	// End of IHapticDevice Interface
 
+	// Pointer to currently active SteamVR HMD Module
+	ISteamVRPlugin* SteamVRHMDModule = nullptr;
+
 	/** Initialize the SteamVR System. Will cause a reconnect if one is already active  */
 	void InitSteamVRSystem();
 
@@ -485,4 +489,8 @@ private:
 	*/
 	float CurrentDeltaTime;
 
+	/**
+	* Flag to signify if we need a re-init
+	*/
+	bool bSteamVRWasShutdown = false;
 };
diff --git a/Engine/Plugins/Runtime/Steam/SteamVR/Source/SteamVRInputDevice/Private/SteamVRInputDeviceFunctionLibrary.cpp b/Engine/Plugins/Runtime/Steam/SteamVR/Source/SteamVRInputDevice/Private/SteamVRInputDeviceFunctionLibrary.cpp
--- a/Engine/Plugins/Runtime/Steam/SteamVR/Source/SteamVRInputDevice/Private/SteamVRInputDeviceFunctionLibrary.cpp
+++ b/Engine/Plugins/Runtime/Steam/SteamVR/Source/SteamVRInputDevice/Private/SteamVRInputDeviceFunctionLibrary.cpp
@@ -761,14 +761,6 @@ void USteamVRInputDeviceFunctionLibrary::GetFingerCurlsAndSplays(EHand Hand, FSt
 			k_ulInvalidActionSetHandle
 		};
 
-		EVRInputError UpdateActionStateError = VRInput()->UpdateActionState(&ActiveActionSet, sizeof(VRActiveActionSet_t), 1);
-		if (UpdateActionStateError != VRInputError_None)
-		{
-			FingerCurls = {};
-			FingerSplays = {};
-			return;
-		}
-
 		// Setup which hand data we will get from SteamVR
 		VRActionHandle_t ActiveSkeletalHand = k_ulInvalidActionHandle;
 
@@ -853,7 +845,14 @@ FSteamVRInputDevice* USteamVRInputDeviceFunctionLibrary::GetSteamVRInputDevice()
 		FSteamVRInputDevice* TestSteamVRDevice = static_cast<FSteamVRInputDevice*>(MotionController);
 		if (TestSteamVRDevice != nullptr && TestSteamVRDevice->DeviceSignature == 2019)
 		{
-			return TestSteamVRDevice;
+			if (TestSteamVRDevice->SteamVRHMDModule && TestSteamVRDevice->SteamVRHMDModule->GetVRSystem())
+			{
+				return TestSteamVRDevice;
+			}
+			else
+			{
+				return nullptr;
+			}
 		}
 	}
 	return nullptr;
diff --git a/Engine/Plugins/Runtime/Steam/SteamVR/Source/SteamVRInputDevice/Private/SteamVRInputDeviceModule.cpp b/Engine/Plugins/Runtime/Steam/SteamVR/Source/SteamVRInputDevice/Private/SteamVRInputDeviceModule.cpp
--- a/Engine/Plugins/Runtime/Steam/SteamVR/Source/SteamVRInputDevice/Private/SteamVRInputDeviceModule.cpp
+++ b/Engine/Plugins/Runtime/Steam/SteamVR/Source/SteamVRInputDevice/Private/SteamVRInputDeviceModule.cpp
@@ -72,23 +72,6 @@ void FSteamVRInputDeviceModule::StartupModule()
 {
 	IModularFeatures::Get().RegisterModularFeature(GetModularFeatureName(), this);
 
-#if WITH_EDITOR
-	EVRInitError SteamVRInitErrorOverlay = VRInitError_Driver_NotLoaded;
-	IVRSystem* SteamVRSystemOverlay = vr::VR_Init(&SteamVRInitErrorOverlay, vr::VRApplication_Overlay);
-
-	if (SteamVRInitErrorOverlay == VRInitError_None)
-	{
-		if (VRSettings() != nullptr)
-		{
-			EVRSettingsError BindingFlagError = VRSettingsError_None;
-			VRSettings()->SetBool(k_pch_SteamVR_Section, k_pch_SteamVR_DebugInputBinding, true, &BindingFlagError);
-			UE_LOG(LogTemp, Display, TEXT("[STEAMVR INPUT] Enable SteamVR Input Developer Mode: %s"), *FString(UTF8_TO_TCHAR(VRSettings()->GetSettingsErrorNameFromEnum(BindingFlagError))));
-			VRSettings()->SetBool(k_pch_SteamVR_Section, k_pch_SteamVR_DebugInput, true, &BindingFlagError);
-			UE_LOG(LogTemp, Display, TEXT("[STEAMVR INPUT] Enable SteamVR Debug Input: %s"), *FString(UTF8_TO_TCHAR(VRSettings()->GetSettingsErrorNameFromEnum(BindingFlagError))));
-		}
-		VR_Shutdown();
-	}
-#endif
 
 	// Unload engine integrated controller modules if Valve's Input Plugin is present
 	FModuleManager& ModuleManager = FModuleManager::Get();
