Roll back SteamVR Input binding file EngineDir/ProjectDir move

diff --git a/Engine/Plugins/Runtime/Steam/SteamVR/Source/SteamVRInputDevice/Private/SteamVRInputDevice.cpp b/Engine/Plugins/Runtime/Steam/SteamVR/Source/SteamVRInputDevice/Private/SteamVRInputDevice.cpp
--- a/Engine/Plugins/Runtime/Steam/SteamVR/Source/SteamVRInputDevice/Private/SteamVRInputDevice.cpp
+++ b/Engine/Plugins/Runtime/Steam/SteamVR/Source/SteamVRInputDevice/Private/SteamVRInputDevice.cpp
@@ -1359,11 +1359,7 @@ void FSteamVRInputDevice::ReloadActionManifest()
 	if (VRSystem() && VRInput() && VRApplications())
 	{
 		// Set Action Manifest Path
-#if WITH_EDITOR
 		const FString ManifestPath = FPaths::ProjectConfigDir() / CONTROLLER_BINDING_PATH / ACTION_MANIFEST;
-#else
-		const FString ManifestPath = FPaths::EngineConfigDir() / CONTROLLER_BINDING_PATH / ACTION_MANIFEST;
-#endif
 		UE_LOG(LogSteamVRInputDevice, Display, TEXT("Reloading Action Manifest in: %s"), *ManifestPath);
 			
 		// Load application manifest
@@ -3459,7 +3455,7 @@ void FSteamVRInputDevice::RegisterApplication(FString ManifestPath)
 		#if WITH_EDITOR
 			TheActionManifestPath = IFileManager::Get().ConvertToAbsolutePathForExternalAppForRead(*ManifestPath);
 		#else
-			TheActionManifestPath = FPaths::ConvertRelativePathToFull(FPaths::EngineDir() / TEXT("Config") / TEXT("SteamVRBindings") / TEXT(ACTION_MANIFEST)).Replace(TEXT("/"), TEXT("\\"));
+			TheActionManifestPath = FPaths::ConvertRelativePathToFull(FPaths::ProjectDir() / TEXT("Config") / TEXT("SteamVRBindings") / TEXT(ACTION_MANIFEST)).Replace(TEXT("/"), TEXT("\\"));
 		#endif
 		
 		UE_LOG(LogSteamVRInputDevice, Display, TEXT("[STEAMVR INPUT] Trying to load Action Manifest from: %s"), *TheActionManifestPath);
