DRPS - Fix some compile issues
#rnx

diff --git a/Engine/Plugins/Runtime/Steam/SteamVR/Source/SteamVR/Private/SteamVRHMD.h b/Engine/Plugins/Runtime/Steam/SteamVR/Source/SteamVR/Private/SteamVRHMD.h
--- a/Engine/Plugins/Runtime/Steam/SteamVR/Source/SteamVR/Private/SteamVRHMD.h
+++ b/Engine/Plugins/Runtime/Steam/SteamVR/Source/SteamVR/Private/SteamVRHMD.h
@@ -13,6 +13,7 @@
 #include "StereoLayerManager.h"
 #include "XRRenderTargetManager.h"
 #include "XRRenderBridge.h"
+#include "XRSwapChain.h"
 #include "IHeadMountedDisplayVulkanExtensions.h"
 
 #if PLATFORM_WINDOWS
@@ -66,52 +67,6 @@ struct FSteamVRLayer
 	friend void MarkLayerTextureForUpdate(FSteamVRLayer& Layer);
 };
 
-/**
- * Render target swap chain
- */
-class FRHITextureSet2D : public FRHITexture2D
-{
-public:
-	
-	FRHITextureSet2D(const uint32 TextureSetSize, EPixelFormat Format, uint32 SizeX, uint32 SizeY, uint32 NumMips, uint32 NumSamples, uint32 Flags, const FClearValueBinding& InClearValue)
-	: FRHITexture2D(SizeX, SizeY, NumMips, NumSamples, Format, Flags, InClearValue)
-	, TextureIndex(0)
-	{
-		TextureSet.AddZeroed(TextureSetSize);
-	}
-	
-	virtual ~FRHITextureSet2D()
-	{}
-	
-	void AddTexture(FTexture2DRHIRef& Texture, const uint32 Index)
-	{
-		check(Index < static_cast<uint32>(TextureSet.Num()));
-		// todo: Check texture format to ensure it matches the set
-		TextureSet[Index] = Texture;
-	}
-	
-	void Advance()
-	{
-		TextureIndex = (TextureIndex + 1) % static_cast<uint32>(TextureSet.Num());
-	}
-	
-	virtual void* GetTextureBaseRHI() override
-	{
-		check(TextureSet[TextureIndex].IsValid());
-		return TextureSet[TextureIndex]->GetTextureBaseRHI();
-	}
-	
-	virtual void* GetNativeResource() const override
-	{
-		check(TextureSet[TextureIndex].IsValid());
-		return TextureSet[TextureIndex]->GetNativeResource();
-	}
-	
-private:
-	TArray<FTexture2DRHIRef> TextureSet;
-	uint32 TextureIndex;
-};
-
 /**
  * SteamVR Head Mounted Display
  */
@@ -180,8 +135,6 @@ public:
 
 	virtual void UpdateScreenSettings(const FViewport* InViewport) override {}
 	
-	virtual bool AllocateRenderTargetTexture(uint32 Index, uint32 SizeX, uint32 SizeY, uint8 Format, uint32 NumMips, uint32 InTexFlags, uint32 InTargetableTextureFlags, FTexture2DRHIRef& OutTargetableTexture, FTexture2DRHIRef& OutShaderResourceTexture, uint32 NumSamples = 1) override;
-
 	virtual bool GetHMDDistortionEnabled(EShadingPath ShadingPath) const override;
 
 	virtual void OnBeginRendering_GameThread() override;
@@ -211,6 +164,9 @@ public:
 	}
 	virtual void CalculateRenderTargetSize(const class FViewport& Viewport, uint32& InOutSizeX, uint32& InOutSizeY) override;
 	virtual bool NeedReAllocateViewportRenderTarget(const class FViewport& Viewport) override;
+	virtual bool NeedReAllocateDepthTexture(const TRefCountPtr<struct IPooledRenderTarget>& DepthTarget) override;
+	virtual bool AllocateRenderTargetTexture(uint32 Index, uint32 SizeX, uint32 SizeY, uint8 Format, uint32 NumMips, uint32 InTexFlags, uint32 InTargetableTextureFlags, FTexture2DRHIRef& OutTargetableTexture, FTexture2DRHIRef& OutShaderResourceTexture, uint32 NumSamples = 1) override;
+	virtual bool AllocateDepthTexture(uint32 Index, uint32 SizeX, uint32 SizeY, uint8 Format, uint32 NumMips, uint32 Flags, uint32 TargetableTextureFlags, FTexture2DRHIRef& OutTargetableTexture, FTexture2DRHIRef& OutShaderResourceTexture, uint32 NumSamples = 1) override;
 
 	// IStereoLayers interface
 	// Create/Set/Get/Destroy inherited from TStereoLayerManager
@@ -255,28 +211,35 @@ public:
 			return bUseExplicitTimingMode;
 		}
 
+		FXRSwapChainPtr GetSwapChain() { return SwapChain; }
+		FXRSwapChainPtr GetDepthSwapChain() { return DepthSwapChain; }
+
 		/** Schedules BeginRendering_RHI on the RHI thread when in explicit timing mode */
 		void BeginRendering_RenderThread(FRHICommandListImmediate& RHICmdList);
 
 		/** Called only when we're in explicit timing mode, which needs to be paired with a call to PostPresentHandoff */
 		void BeginRendering_RHI();
 
+		void CreateSwapChain(const FTextureRHIRef& BindingTexture, TArray<FTextureRHIRef>&& SwapChainTextures);
+		void CreateDepthSwapChain(const FTextureRHIRef& BindingTexture, TArray<FTextureRHIRef>&& SwapChainTextures);
 
 		// Virtual interface implemented by subclasses
 		virtual void Reset() = 0;
+
 	private:
 		virtual void FinishRendering() = 0;
 
 	protected:
 		
 		FSteamVRHMD*			Plugin;
+		FXRSwapChainPtr			SwapChain;
+		FXRSwapChainPtr			DepthSwapChain;
+
 		bool					bInitialized;
 		
 		/** If we use explicit timing mode, we must have matching calls to BeginRendering_RHI and PostPresentHandoff */
 		bool					bUseExplicitTimingMode;
 		bool NeedsPostPresentHandoff() const;
-		
-
 	};
 
 #if PLATFORM_WINDOWS
@@ -285,13 +248,9 @@ public:
 	public:
 		D3D11Bridge(FSteamVRHMD* plugin);
 
-
 		virtual void FinishRendering() override;
 		virtual void UpdateViewport(const FViewport& Viewport, FRHIViewport* InViewportRHI) override;
 		virtual void Reset() override;
-
-	protected:
-		ID3D11Texture2D* RenderTargetTexture = NULL;
 	};
 #endif // PLATFORM_WINDOWS
 
@@ -318,10 +277,6 @@ public:
 		virtual void FinishRendering() override;
 		virtual void UpdateViewport(const FViewport& Viewport, FRHIViewport* InViewportRHI) override;
 		virtual void Reset() override;
-
-	protected:
-
-		FTexture2DRHIRef RenderTargetTexture;
 	};
 
 	class OpenGLBridge : public BridgeBaseImpl
@@ -332,10 +287,6 @@ public:
 		virtual void FinishRendering() override;
 		virtual void UpdateViewport(const FViewport& Viewport, FRHIViewport* InViewportRHI) override;
 		virtual void Reset() override;
-
-	protected:
-		GLuint RenderTargetTexture = 0;
-
 	};
 	
 #elif PLATFORM_MAC
@@ -349,8 +300,6 @@ public:
 		virtual void Reset() override;
 		
 		IOSurfaceRef GetSurface(const uint32 SizeX, const uint32 SizeY);
-		
-		FTexture2DRHIRef TextureSet;
 	};
 #endif // PLATFORM_MAC
 
@@ -441,6 +390,36 @@ public:
 
 		return out;
 	}
+
+	static FORCEINLINE vr::HmdMatrix44_t ToHmdMatrix44(const FMatrix& tm)
+	{
+		// Rows and columns are swapped between vr::HmdMatrix44_t and FMatrix
+		vr::HmdMatrix44_t out;
+
+		out.m[0][0] = tm.M[0][0];
+		out.m[1][0] = tm.M[0][1];
+		out.m[2][0] = tm.M[0][2];
+		out.m[3][0] = tm.M[0][3];
+
+		out.m[0][1] = tm.M[1][0];
+		out.m[1][1] = tm.M[1][1];
+		out.m[2][1] = tm.M[1][2];
+		out.m[3][1] = tm.M[1][3];
+
+		out.m[0][2] = tm.M[2][0];
+		out.m[1][2] = tm.M[2][1];
+		out.m[2][2] = tm.M[2][2];
+		out.m[3][2] = tm.M[2][3];
+
+		out.m[0][3] = tm.M[3][0];
+		out.m[1][3] = tm.M[3][1];
+		out.m[2][3] = tm.M[3][2];
+		out.m[3][3] = tm.M[3][3];
+
+		return out;
+	}
+
+
 private:
 
 	void SetupOcclusionMeshes();
diff --git a/Engine/Plugins/Runtime/Steam/SteamVR/Source/SteamVR/Private/SteamVRRender.cpp b/Engine/Plugins/Runtime/Steam/SteamVR/Source/SteamVR/Private/SteamVRRender.cpp
--- a/Engine/Plugins/Runtime/Steam/SteamVR/Source/SteamVR/Private/SteamVRRender.cpp
+++ b/Engine/Plugins/Runtime/Steam/SteamVR/Source/SteamVR/Private/SteamVRRender.cpp
@@ -68,9 +68,10 @@ bool FSteamVRHMD::IsActiveThisFrame(class FViewport* InViewport) const
 static void DrawOcclusionMesh(FRHICommandList& RHICmdList, EStereoscopicPass StereoPass, const FHMDViewMesh MeshAssets[])
 {
 	check(IsInRenderingThread());
-	check(IStereoRendering::IsStereoEyeView(StereoPass));
+	check(StereoPass == eSSP_LEFT_EYE);
 
-	const uint32 MeshIndex = IStereoRendering::IsAPrimaryView(StereoPass) ? 0 : 1;
+	//#todo-rv-vr
+	const uint32 MeshIndex = /*IStereoRendering::IsAPrimaryView*/(StereoPass == eSSP_FULL || StereoPass == eSSP_LEFT_EYE) ? 0 : 1;
 	const FHMDViewMesh& Mesh = MeshAssets[MeshIndex];
 	check(Mesh.IsValid());
 
diff --git a/Engine/Source/Runtime/Renderer/Private/LightPropagationVolume.cpp b/Engine/Source/Runtime/Renderer/Private/LightPropagationVolume.cpp
--- a/Engine/Source/Runtime/Renderer/Private/LightPropagationVolume.cpp
+++ b/Engine/Source/Runtime/Renderer/Private/LightPropagationVolume.cpp
@@ -1459,11 +1459,13 @@ void FSceneViewState::SetupLightPropagationVolume(FSceneView& View, FSceneViewFa
 
 	const ERHIFeatureLevel::Type ViewFeatureLevel = View.GetFeatureLevel();
 
-	if (IStereoRendering::IsASecondaryView(View.StereoPass))
+	//#todo-rv-vr
+	if (IStereoRendering::IsASecondaryView(View, GEngine->StereoRenderingDevice))
 	{
 		// The right eye will reference the left eye's LPV with the assumption that the left eye uses the primary view (index 0)
 		const FSceneView* PrimaryView = ViewFamily.Views[0];
-		if (IStereoRendering::IsAPrimaryView(PrimaryView->StereoPass) && PrimaryView->State)
+		//#todo-rv-vr
+		if (PrimaryView->StereoPass == eSSP_LEFT_EYE && PrimaryView->State)
 		{
 			FSceneViewState* PrimaryViewState = PrimaryView->State->GetConcreteViewState();
 			if (PrimaryViewState)
diff --git a/Engine/Source/Runtime/Renderer/Private/PostProcess/PostProcessMobile.cpp b/Engine/Source/Runtime/Renderer/Private/PostProcess/PostProcessMobile.cpp
--- a/Engine/Source/Runtime/Renderer/Private/PostProcess/PostProcessMobile.cpp
+++ b/Engine/Source/Runtime/Renderer/Private/PostProcess/PostProcessMobile.cpp
@@ -2952,7 +2952,8 @@ void FRCPassPostProcessAaES2::Process(FRenderingCompositePassContext& Context)
 	check(SrcSize == DestSize);
 
 	ERenderTargetActions LoadStoreAction = ERenderTargetActions::Load_Store;
-	if (!IStereoRendering::IsASecondaryView(Context.View.StereoPass))
+	//#todo-rv-vr
+	if (Context.View.StereoPass != eSSP_RIGHT_EYE)
 	{
 		// Full clear to avoid restore
 		LoadStoreAction = ERenderTargetActions::Clear_Store;
