Integration issue from stereo rendering interface changes.

diff --git a/Engine/Plugins/Runtime/Steam/SteamVR/Source/SteamVR/Private/SteamVRRender.cpp b/Engine/Plugins/Runtime/Steam/SteamVR/Source/SteamVR/Private/SteamVRRender.cpp
--- a/Engine/Plugins/Runtime/Steam/SteamVR/Source/SteamVR/Private/SteamVRRender.cpp
+++ b/Engine/Plugins/Runtime/Steam/SteamVR/Source/SteamVR/Private/SteamVRRender.cpp
@@ -68,10 +68,9 @@ bool FSteamVRHMD::IsActiveThisFrame(class FViewport* InViewport) const
 static void DrawOcclusionMesh(FRHICommandList& RHICmdList, EStereoscopicPass StereoPass, const FHMDViewMesh MeshAssets[])
 {
 	check(IsInRenderingThread());
-	check(StereoPass == eSSP_LEFT_EYE);
+	check(GEngine->StereoRenderingDevice->IsStereoEyePass(StereoPass));
 
-	//#todo-rv-vr
-	const uint32 MeshIndex = /*IStereoRendering::IsAPrimaryView*/(StereoPass == eSSP_FULL || StereoPass == eSSP_LEFT_EYE) ? 0 : 1;
+	const uint32 MeshIndex = GEngine->StereoRenderingDevice->GetViewIndexForPass(StereoPass);
 	const FHMDViewMesh& Mesh = MeshAssets[MeshIndex];
 	check(Mesh.IsValid());
 
