Implement SteamVR high-quality stereo layers by implementing a plugin-specific layer shape.

diff --git a/Engine/Plugins/Runtime/Steam/SteamVR/Source/SteamVR/Classes/SteamVRHQLayerShape.h b/Engine/Plugins/Runtime/Steam/SteamVR/Source/SteamVR/Classes/SteamVRHQLayerShape.h
new file mode 100644
--- /dev/null
+++ b/Engine/Plugins/Runtime/Steam/SteamVR/Source/SteamVR/Classes/SteamVRHQLayerShape.h
@@ -0,0 +1,38 @@
+// Copyright 1998-2019 Epic Games, Inc. All Rights Reserved.
+
+#pragma once
+#include "Components/StereoLayerComponent.h"
+#include "SteamVRHQLayerShape.generated.h"
+
+UCLASS(meta = (DisplayName = "High Quality Layer (SteamVR)"))
+class STEAMVR_API USteamVRHQStereoLayerShape : public UStereoLayerShapeQuad
+{
+	GENERATED_BODY()
+public:
+	virtual void ApplyShape(IStereoLayers::FLayerDesc& LayerDesc) override;
+
+	UPROPERTY(EditAnywhere, BlueprintReadOnly, export, Category = "HQ Layer Properties")
+	bool bCurved;
+	
+	UPROPERTY(EditAnywhere, BlueprintReadOnly, export, Category = "HQ Layer Properties")
+	bool bAntiAlias;
+	
+	/** For curved layers, sets the minimum distance from the layer used to automatically curve
+      * the surface around the viewer.  The minimum distance is when the layer is most curved. */
+	UPROPERTY(EditAnywhere, BlueprintReadOnly, export, Category = "HQ Layer Properties")
+	float AutoCurveMinDistance;
+	
+	/** For curved layers, sets the maximum distance from the layer used to automatically curve
+	  * the surface around the viewer.  The maximum distance is when the layer is the least curved. */
+	UPROPERTY(EditAnywhere, BlueprintReadOnly, export, Category = "HQ Layer Properties")
+	float AutoCurveMaxDistance;
+
+	UFUNCTION(BlueprintCallable, Category = "Components|Stereo Layer")
+	void SetCurved(bool InCurved);
+	UFUNCTION(BlueprintCallable, Category = "Components|Stereo Layer")
+	void SetAntiAlias(bool InAntiAlias);
+	UFUNCTION(BlueprintCallable, Category = "Components|Stereo Layer")
+	void SetAutoCurveMinDistance(float InDistance);
+	UFUNCTION(BlueprintCallable, Category = "Components|Stereo Layer")
+	void SetAutoCurveMaxDistance(float InDistance);
+};
\ No newline at end of file
diff --git a/Engine/Plugins/Runtime/Steam/SteamVR/Source/SteamVR/Private/SteamVRHQLayerShape.cpp b/Engine/Plugins/Runtime/Steam/SteamVR/Source/SteamVR/Private/SteamVRHQLayerShape.cpp
new file mode 100644
--- /dev/null
+++ b/Engine/Plugins/Runtime/Steam/SteamVR/Source/SteamVR/Private/SteamVRHQLayerShape.cpp
@@ -0,0 +1,51 @@
+// Copyright 1998-2019 Epic Games, Inc. All Rights Reserved.
+
+#include "SteamVRHQLayerShape.h"
+#include "SteamVRStereoLayers.h"
+#include "IStereoLayers.h"
+
+void USteamVRHQStereoLayerShape::ApplyShape(IStereoLayers::FLayerDesc& LayerDesc)
+{
+	LayerDesc.SetShape<FSteamVRHQLayer>(bCurved, bAntiAlias, AutoCurveMinDistance, AutoCurveMaxDistance);
+}
+
+void USteamVRHQStereoLayerShape::SetCurved(bool InCurved)
+{
+	if (InCurved == bCurved)
+	{
+		return;
+	}
+	bCurved = InCurved;
+	MarkStereoLayerDirty();
+}
+
+void USteamVRHQStereoLayerShape::SetAntiAlias(bool InAntiAlias)
+{
+	if (InAntiAlias == bAntiAlias)
+	{
+		return;
+	}
+	bAntiAlias = InAntiAlias;
+	MarkStereoLayerDirty();
+}
+
+void USteamVRHQStereoLayerShape::SetAutoCurveMinDistance(float InDistance)
+{
+	if (InDistance == AutoCurveMinDistance)
+	{
+		return;
+	}
+	AutoCurveMinDistance = InDistance;
+	MarkStereoLayerDirty();
+}
+
+void USteamVRHQStereoLayerShape::SetAutoCurveMaxDistance(float InDistance)
+{
+	if (InDistance == AutoCurveMaxDistance)
+	{
+		return;
+	}
+	AutoCurveMaxDistance = InDistance;
+	MarkStereoLayerDirty();
+}
+
diff --git a/Engine/Plugins/Runtime/Steam/SteamVR/Source/SteamVR/Private/SteamVRStereoLayers.cpp b/Engine/Plugins/Runtime/Steam/SteamVR/Source/SteamVR/Private/SteamVRStereoLayers.cpp
--- a/Engine/Plugins/Runtime/Steam/SteamVR/Source/SteamVR/Private/SteamVRStereoLayers.cpp
+++ b/Engine/Plugins/Runtime/Steam/SteamVR/Source/SteamVR/Private/SteamVRStereoLayers.cpp
@@ -2,11 +2,14 @@
 //
 #include "CoreMinimal.h"
 #include "SteamVRPrivate.h"
+#include "SteamVRStereoLayers.h"
 #include "StereoLayerManager.h"
 #include "SteamVRHMD.h"
 #include "Misc/ScopeLock.h"
 #include "DefaultXRCamera.h"
 
+const FName FSteamVRHQLayer::ShapeName = FName("SteamVRHQLayer");
+
 #if STEAMVR_SUPPORTED_PLATFORMS
 
 // enable to ensure all openvr calls succeed
@@ -95,9 +98,33 @@ void FSteamVRHMD::UpdateLayer(struct FSteamVRLayer& Layer, uint32 LayerId, bool
 	{
 		UE_LOG(LogHMD, Warning, TEXT("Unsupported StereoLayer flag. SteamVR StereoLayers do not support disabling alpha renderding. Make the texture opaque instead."));
 	}
-	if (!Layer.LayerDesc.HasShape<FQuadLayer>())
+
+	const float WorldToMeterScale = GetWorldToMetersScale();
+	check(WorldToMeterScale > 0.f);
+
+	if (Layer.LayerDesc.HasShape<FSteamVRHQLayer>())
+	{
+		vr::VROverlayHandle_t ExistingHQLayer = VROverlay->GetHighQualityOverlay();
+		if (ExistingHQLayer != vr::k_ulOverlayHandleInvalid && ExistingHQLayer != Layer.OverlayHandle)
+		{
+			UE_LOG(LogHMD, Warning, TEXT("There can only be one high-quality stereo layer active at a time. Please disable or change one or more layers to non-high-quality."));
+		}
+		else
+		{
+			VROverlay->SetHighQualityOverlay(Layer.OverlayHandle);
+
+			const auto& HQSettings = Layer.LayerDesc.GetShape<FSteamVRHQLayer>();
+			VROverlay->SetOverlayFlag(Layer.OverlayHandle, vr::VROverlayFlags_Curved, HQSettings.bCurved);
+			VROverlay->SetOverlayFlag(Layer.OverlayHandle, vr::VROverlayFlags_RGSS4X, HQSettings.bAntiAlias);
+			if (HQSettings.bCurved)
+			{
+				VROverlay->SetOverlayAutoCurveDistanceRangeInMeters(Layer.OverlayHandle, HQSettings.AutoCurveMinDistance / WorldToMeterScale, HQSettings.AutoCurveMaxDistance / WorldToMeterScale);
+			}
+		}
+	}
+	else if (!Layer.LayerDesc.HasShape<FQuadLayer>())
 	{
-		UE_LOG(LogHMD, Warning, TEXT("Unsupported StereoLayer shape. SteamVR StereoLayers can only be Quads."));
+		UE_LOG(LogHMD, Warning, TEXT("Unsupported StereoLayer shape. SteamVR StereoLayers can only be Quads or High Quality Quads."));
 	}
 
 	// UVs
@@ -107,8 +134,6 @@ void FSteamVRHMD::UpdateLayer(struct FSteamVRLayer& Layer, uint32 LayerId, bool
     TextureBounds.vMin = Layer.LayerDesc.UVRect.Min.Y;
     TextureBounds.vMax = Layer.LayerDesc.UVRect.Max.Y;
 	OVR_VERIFY(VROverlay->SetOverlayTextureBounds(Layer.OverlayHandle, &TextureBounds));
-	const float WorldToMeterScale = GetWorldToMetersScale();
-	check(WorldToMeterScale > 0.f);
 	OVR_VERIFY(VROverlay->SetOverlayWidthInMeters(Layer.OverlayHandle, Layer.LayerDesc.QuadSize.X / WorldToMeterScale));
 	
 	float TexelAspect = 1.0f;
diff --git a/Engine/Plugins/Runtime/Steam/SteamVR/Source/SteamVR/Public/SteamVRStereoLayers.h b/Engine/Plugins/Runtime/Steam/SteamVR/Source/SteamVR/Public/SteamVRStereoLayers.h
new file mode 100644
--- /dev/null
+++ b/Engine/Plugins/Runtime/Steam/SteamVR/Source/SteamVR/Public/SteamVRStereoLayers.h
@@ -0,0 +1,23 @@
+// Copyright 1998-2019 Epic Games, Inc. All Rights Reserved.
+#pragma once
+#include "StereoLayerShapes.h"
+
+/** Class describing additional settings for high-quality layers. Note there can only be one HQ layer active at any given time. */
+class STEAMVR_API FSteamVRHQLayer : public IStereoLayerShape
+{
+	STEREO_LAYER_SHAPE_BOILERPLATE(FSteamVRHQLayer)
+public:
+
+	FSteamVRHQLayer() {}
+	FSteamVRHQLayer(bool bInCurved, bool bInAntiAlias, float InAutoCurveMinDistance, float InAutoCurveMaxDistance) :
+		bCurved(bInCurved),
+		bAntiAlias(bInAntiAlias),
+		AutoCurveMinDistance(InAutoCurveMinDistance),
+		AutoCurveMaxDistance(InAutoCurveMaxDistance)
+	{}
+
+	bool bCurved;
+	bool bAntiAlias;
+	float AutoCurveMinDistance;
+	float AutoCurveMaxDistance;
+};
\ No newline at end of file
diff --git a/Engine/Source/Runtime/Engine/Classes/Components/StereoLayerComponent.h b/Engine/Source/Runtime/Engine/Classes/Components/StereoLayerComponent.h
--- a/Engine/Source/Runtime/Engine/Classes/Components/StereoLayerComponent.h
+++ b/Engine/Source/Runtime/Engine/Classes/Components/StereoLayerComponent.h
@@ -46,7 +46,7 @@ enum EStereoLayerShape
 };
 
 UCLASS(EditInlineNew, Abstract, BlueprintType, CollapseCategories, Within = StereoLayerComponent)
-class UStereoLayerShape : public UObject
+class ENGINE_API UStereoLayerShape : public UObject
 {
 	GENERATED_BODY()
 
@@ -61,7 +61,7 @@ protected:
 };
 
 UCLASS(meta = (DisplayName = "Quad Layer"))
-class UStereoLayerShapeQuad : public UStereoLayerShape
+class ENGINE_API UStereoLayerShapeQuad : public UStereoLayerShape
 {
 	GENERATED_BODY()
 public:
@@ -72,7 +72,7 @@ public:
 };
 
 UCLASS(meta = (DisplayName = "Cylinder Layer"))
-class UStereoLayerShapeCylinder : public UStereoLayerShape
+class ENGINE_API UStereoLayerShapeCylinder : public UStereoLayerShape
 {
 	GENERATED_BODY()
 
@@ -109,7 +109,7 @@ public:
 };
 
 UCLASS(meta = (DisplayName = "Cubemap Layer"))
-class UStereoLayerShapeCubemap : public UStereoLayerShape
+class ENGINE_API UStereoLayerShapeCubemap : public UStereoLayerShape
 {
 	GENERATED_BODY()
 public:
@@ -186,7 +186,7 @@ public:
 
 
 UCLASS(meta = (DisplayName = "Equirect Layer"))
-class UStereoLayerShapeEquirect : public UStereoLayerShape
+class ENGINE_API UStereoLayerShapeEquirect : public UStereoLayerShape
 {
 	GENERATED_BODY()
 
