PR #6680: SteamVR Input Hotfix 3 Candidate Fixes (Contributed by [REDACTED_PERSON])

(1) Allow tracking without input mappings
(2) Trackpad Touch key binding will now trigger properly without press
(3) New Trackpad Dpad Touch key support
(4) Removed the default slight pulse for Trigger and Grip on Index & Touch controllers (as per dev requests in UDN)
(5) Fallback to SteamVR tracker render model for unrecognized motion sources

diff --git a/Engine/Plugins/Runtime/Steam/SteamVR/Source/SteamVR/Private/SteamVRAssetManager.cpp b/Engine/Plugins/Runtime/Steam/SteamVR/Source/SteamVR/Private/SteamVRAssetManager.cpp
--- a/Engine/Plugins/Runtime/Steam/SteamVR/Source/SteamVR/Private/SteamVRAssetManager.cpp
+++ b/Engine/Plugins/Runtime/Steam/SteamVR/Source/SteamVR/Private/SteamVRAssetManager.cpp
@@ -696,22 +696,8 @@ int32 FSteamVRAssetManager::GetDeviceId(EControllerHand ControllerHand)
 			DesiredDeviceClass = vr::ETrackedDeviceClass::TrackedDeviceClass_TrackingReference;
 			break;
 
-		case EControllerHand::Special_1:
-		case EControllerHand::Special_2:
-		case EControllerHand::Special_3:
-		case EControllerHand::Special_4:
-		case EControllerHand::Special_5:
-		case EControllerHand::Special_6:
-		case EControllerHand::Special_7:
-		case EControllerHand::Special_8:
-		case EControllerHand::Special_9:
-		case EControllerHand::Special_10:
-		case EControllerHand::Special_11:
-			DesiredDeviceClass = vr::ETrackedDeviceClass::TrackedDeviceClass_GenericTracker;
-			break;
-
 		default:
-			// DesiredDeviceClass = TrackedDeviceClass_Invalid => returns -1
+			DesiredDeviceClass = vr::ETrackedDeviceClass::TrackedDeviceClass_GenericTracker;
 			break;
 		}
 
diff --git a/Engine/Plugins/Runtime/Steam/SteamVR/Source/SteamVRInputDevice/Private/SteamVRControllerKeys.h b/Engine/Plugins/Runtime/Steam/SteamVR/Source/SteamVRInputDevice/Private/SteamVRControllerKeys.h
--- a/Engine/Plugins/Runtime/Steam/SteamVR/Source/SteamVRInputDevice/Private/SteamVRControllerKeys.h
+++ b/Engine/Plugins/Runtime/Steam/SteamVR/Source/SteamVRInputDevice/Private/SteamVRControllerKeys.h
@@ -93,3 +93,38 @@ namespace CosmosKeys
 	const FKey Cosmos_Right_Thumbstick_Right("Cosmos_Right_Thumbstick_Right");
 }
 
+namespace InputKeys
+{
+	// Valve Index - Additional input keys not implemented in OpenXR yet
+	const FKey ValveIndex_Left_Trackpad_Up_Touch( "ValveIndex_Left_Trackpad_Up_Touch" );
+	const FKey ValveIndex_Left_Trackpad_Down_Touch( "ValveIndex_Left_Trackpad_Down_Touch" );
+	const FKey ValveIndex_Left_Trackpad_Left_Touch( "ValveIndex_Left_Trackpad_Left_Touch" );
+	const FKey ValveIndex_Left_Trackpad_Right_Touch( "ValveIndex_Left_Trackpad_Right_Touch" );
+
+	const FKey ValveIndex_Right_Trackpad_Up_Touch( "ValveIndex_Right_Trackpad_Up_Touch" );
+	const FKey ValveIndex_Right_Trackpad_Down_Touch( "ValveIndex_Right_Trackpad_Down_Touch" );
+	const FKey ValveIndex_Right_Trackpad_Left_Touch( "ValveIndex_Right_Trackpad_Left_Touch" );
+	const FKey ValveIndex_Right_Trackpad_Right_Touch( "ValveIndex_Right_Trackpad_Right_Touch" );
+
+	// HTC Vive - Additional input keys not implemented in OpenXR yet
+	const FKey Vive_Left_Trackpad_Up_Touch( "Vive_Left_Trackpad_Up_Touch" );
+	const FKey Vive_Left_Trackpad_Down_Touch( "Vive_Left_Trackpad_Down_Touch" );
+	const FKey Vive_Left_Trackpad_Left_Touch( "Vive_Left_Trackpad_Left_Touch" );
+	const FKey Vive_Left_Trackpad_Right_Touch( "Vive_Left_Trackpad_Right_Touch" );
+
+	const FKey Vive_Right_Trackpad_Up_Touch( "Vive_Right_Trackpad_Up_Touch" );
+	const FKey Vive_Right_Trackpad_Down_Touch( "Vive_Right_Trackpad_Down_Touch" );
+	const FKey Vive_Right_Trackpad_Left_Touch( "Vive_Right_Trackpad_Left_Touch" );
+	const FKey Vive_Right_Trackpad_Right_Touch( "Vive_Right_Trackpad_Right_Touch" );
+
+	// Windows Mixed Reality - Additional input keys not implemented in OpenXR yet
+	const FKey MixedReality_Left_Trackpad_Up_Touch( "MixedReality_Left_Trackpad_Up_Touch" );
+	const FKey MixedReality_Left_Trackpad_Down_Touch( "MixedReality_Left_Trackpad_Down_Touch" );
+	const FKey MixedReality_Left_Trackpad_Left_Touch( "MixedReality_Left_Trackpad_Left_Touch" );
+	const FKey MixedReality_Left_Trackpad_Right_Touch( "MixedReality_Left_Trackpad_Right_Touch" );
+
+	const FKey MixedReality_Right_Trackpad_Up_Touch( "MixedReality_Right_Trackpad_Up_Touch" );
+	const FKey MixedReality_Right_Trackpad_Down_Touch( "MixedReality_Right_Trackpad_Down_Touch" );
+	const FKey MixedReality_Right_Trackpad_Left_Touch( "MixedReality_Right_Trackpad_Left_Touch" );
+	const FKey MixedReality_Right_Trackpad_Right_Touch( "MixedReality_Right_Trackpad_Right_Touch" );
+}
diff --git a/Engine/Plugins/Runtime/Steam/SteamVR/Source/SteamVRInputDevice/Private/SteamVRInputDevice.cpp b/Engine/Plugins/Runtime/Steam/SteamVR/Source/SteamVRInputDevice/Private/SteamVRInputDevice.cpp
--- a/Engine/Plugins/Runtime/Steam/SteamVR/Source/SteamVRInputDevice/Private/SteamVRInputDevice.cpp
+++ b/Engine/Plugins/Runtime/Steam/SteamVR/Source/SteamVRInputDevice/Private/SteamVRInputDevice.cpp
@@ -1204,6 +1204,42 @@ void FSteamVRInputDevice::InitControllerKeys()
 	EKeys::AddKey(FKeyDetails(CosmosKeys::Cosmos_Right_Thumbstick_Left, LOCTEXT("Cosmos_Right_Thumbstick_Left", "Cosmos (R) Thumbstick Left"), FKeyDetails::GamepadKey | FKeyDetails::NotBlueprintBindableKey, "Cosmos"));
 	EKeys::AddKey(FKeyDetails(CosmosKeys::Cosmos_Right_Thumbstick_Right, LOCTEXT("Cosmos_Right_Thumbstick_Right", "Cosmos (R) Thumbstick Right"), FKeyDetails::GamepadKey | FKeyDetails::NotBlueprintBindableKey, "Cosmos"));
 #pragma endregion
+
+#pragma region INPUT KEYS (Additional non-standard OpenXR keys)
+	// Valve Index - Additional input keys not implemented in OpenXR yet
+	EKeys::AddKey( FKeyDetails( InputKeys::ValveIndex_Left_Trackpad_Up_Touch, LOCTEXT( "ValveIndex_Left_Trackpad_Up_Touch", "Valve Index (L) Trackpad Up Touch" ), FKeyDetails::GamepadKey | FKeyDetails::NotBlueprintBindableKey, "ValveIndex" ) );
+	EKeys::AddKey( FKeyDetails( InputKeys::ValveIndex_Left_Trackpad_Down_Touch, LOCTEXT( "ValveIndex_Left_Trackpad_Down_Touch", "Valve Index (L) Trackpad Down Touch" ), FKeyDetails::GamepadKey | FKeyDetails::NotBlueprintBindableKey, "ValveIndex" ) );
+	EKeys::AddKey( FKeyDetails( InputKeys::ValveIndex_Left_Trackpad_Left_Touch, LOCTEXT( "ValveIndex_Left_Trackpad_Left_Touch", "Valve Index (L) Trackpad Left Touch" ), FKeyDetails::GamepadKey | FKeyDetails::NotBlueprintBindableKey, "ValveIndex" ) );
+	EKeys::AddKey( FKeyDetails( InputKeys::ValveIndex_Left_Trackpad_Right_Touch, LOCTEXT( "ValveIndex_Left_Trackpad_Right_Touch", "Valve Index (L) Trackpad Right Touch" ), FKeyDetails::GamepadKey | FKeyDetails::NotBlueprintBindableKey, "ValveIndex" ) );
+
+	EKeys::AddKey( FKeyDetails( InputKeys::ValveIndex_Right_Trackpad_Up_Touch, LOCTEXT( "ValveIndex_Right_Trackpad_Up_Touch", "Valve Index (R) Trackpad Up Touch" ), FKeyDetails::GamepadKey | FKeyDetails::NotBlueprintBindableKey, "ValveIndex" ) );
+	EKeys::AddKey( FKeyDetails( InputKeys::ValveIndex_Right_Trackpad_Down_Touch, LOCTEXT( "ValveIndex_Right_Trackpad_Down_Touch", "Valve Index (R) Trackpad Down Touch" ), FKeyDetails::GamepadKey | FKeyDetails::NotBlueprintBindableKey, "ValveIndex" ) );
+	EKeys::AddKey( FKeyDetails( InputKeys::ValveIndex_Right_Trackpad_Left_Touch, LOCTEXT( "ValveIndex_Right_Trackpad_Left_Touch", "Valve Index (R) Trackpad Left Touch" ), FKeyDetails::GamepadKey | FKeyDetails::NotBlueprintBindableKey, "ValveIndex" ) );
+	EKeys::AddKey( FKeyDetails( InputKeys::ValveIndex_Right_Trackpad_Right_Touch, LOCTEXT( "ValveIndex_Right_Trackpad_Right_Touch", "Valve Index (R) Trackpad Right Touch" ), FKeyDetails::GamepadKey | FKeyDetails::NotBlueprintBindableKey, "ValveIndex" ) );
+
+	// HTC Vive - Additional input keys not implemented in OpenXR yet
+	EKeys::AddKey( FKeyDetails( InputKeys::Vive_Left_Trackpad_Up_Touch, LOCTEXT( "Vive_Left_Trackpad_Up", "Vive (L) Trackpad Up Touch" ), FKeyDetails::GamepadKey | FKeyDetails::NotBlueprintBindableKey, "Vive" ) );
+	EKeys::AddKey( FKeyDetails( InputKeys::Vive_Left_Trackpad_Down_Touch, LOCTEXT( "Vive_Left_Trackpad_Down", "Vive (L) Trackpad Down Touch" ), FKeyDetails::GamepadKey | FKeyDetails::NotBlueprintBindableKey, "Vive" ) );
+	EKeys::AddKey( FKeyDetails( InputKeys::Vive_Left_Trackpad_Left_Touch, LOCTEXT( "Vive_Left_Trackpad_Left", "Vive (L) Trackpad Left Touch" ), FKeyDetails::GamepadKey | FKeyDetails::NotBlueprintBindableKey, "Vive" ) );
+	EKeys::AddKey( FKeyDetails( InputKeys::Vive_Left_Trackpad_Right_Touch, LOCTEXT( "Vive_Left_Trackpad_Right", "Vive (L) Trackpad Right Touch" ), FKeyDetails::GamepadKey | FKeyDetails::NotBlueprintBindableKey, "Vive" ) );
+
+	EKeys::AddKey( FKeyDetails( InputKeys::Vive_Right_Trackpad_Up_Touch, LOCTEXT( "Vive_Right_Trackpad_Up", "Vive (R) Trackpad Up Touch" ), FKeyDetails::GamepadKey | FKeyDetails::NotBlueprintBindableKey, "Vive" ) );
+	EKeys::AddKey( FKeyDetails( InputKeys::Vive_Right_Trackpad_Down_Touch, LOCTEXT( "Vive_Right_Trackpad_Down", "Vive (R) Trackpad Down Touch" ), FKeyDetails::GamepadKey | FKeyDetails::NotBlueprintBindableKey, "Vive" ) );
+	EKeys::AddKey( FKeyDetails( InputKeys::Vive_Right_Trackpad_Left_Touch, LOCTEXT( "Vive_Right_Trackpad_Left", "Vive (R) Trackpad Left Touch" ), FKeyDetails::GamepadKey | FKeyDetails::NotBlueprintBindableKey, "Vive" ) );
+	EKeys::AddKey( FKeyDetails( InputKeys::Vive_Right_Trackpad_Right_Touch, LOCTEXT( "Vive_Right_Trackpad_Right", "Vive (R) Trackpad Right Touch" ), FKeyDetails::GamepadKey | FKeyDetails::NotBlueprintBindableKey, "Vive" ) );
+
+	// Windows Mixed Reality - Additional input keys not implemented in OpenXR yet
+	EKeys::AddKey( FKeyDetails( InputKeys::MixedReality_Left_Trackpad_Up_Touch, LOCTEXT( "MixedReality_Left_Trackpad_Up_Touch", "Mixed Reality (L) Trackpad Up Touch" ), FKeyDetails::GamepadKey | FKeyDetails::NotBlueprintBindableKey, "MixedReality" ) );
+	EKeys::AddKey( FKeyDetails( InputKeys::MixedReality_Left_Trackpad_Down_Touch, LOCTEXT( "MixedReality_Left_Trackpad_Down_Touch", "Mixed Reality (L) Trackpad Down Touch" ), FKeyDetails::GamepadKey | FKeyDetails::NotBlueprintBindableKey, "MixedReality" ) );
+	EKeys::AddKey( FKeyDetails( InputKeys::MixedReality_Left_Trackpad_Left_Touch, LOCTEXT( "MixedReality_Left_Trackpad_Left_Touch", "Mixed Reality (L) Trackpad Left Touch" ), FKeyDetails::GamepadKey | FKeyDetails::NotBlueprintBindableKey, "MixedReality" ) );
+	EKeys::AddKey( FKeyDetails( InputKeys::MixedReality_Left_Trackpad_Right_Touch, LOCTEXT( "MixedReality_Left_Trackpad_Right_Touch", "Mixed Reality (L) Trackpad Right Touch" ), FKeyDetails::GamepadKey | FKeyDetails::NotBlueprintBindableKey, "MixedReality" ) );
+
+	EKeys::AddKey( FKeyDetails( InputKeys::MixedReality_Right_Trackpad_Up_Touch, LOCTEXT( "MixedReality_Right_Trackpad_Up_Touch", "Mixed Reality (R) Trackpad Up Touch" ), FKeyDetails::GamepadKey | FKeyDetails::NotBlueprintBindableKey, "MixedReality" ) );
+	EKeys::AddKey( FKeyDetails( InputKeys::MixedReality_Right_Trackpad_Down_Touch, LOCTEXT( "MixedReality_Right_Trackpad_Down_Touch", "Mixed Reality (R) Trackpad Down Touch" ), FKeyDetails::GamepadKey | FKeyDetails::NotBlueprintBindableKey, "MixedReality" ) );
+	EKeys::AddKey( FKeyDetails( InputKeys::MixedReality_Right_Trackpad_Left_Touch, LOCTEXT( "MixedReality_Right_Trackpad_Left_Touch", "Mixed Reality (R) Trackpad Left Touch" ), FKeyDetails::GamepadKey | FKeyDetails::NotBlueprintBindableKey, "MixedReality" ) );
+	EKeys::AddKey( FKeyDetails( InputKeys::MixedReality_Right_Trackpad_Right_Touch, LOCTEXT( "MixedReality_Right_Trackpad_Right_Touch", "Mixed Reality (R) Trackpad Right Touch" ), FKeyDetails::GamepadKey | FKeyDetails::NotBlueprintBindableKey, "MixedReality" ) );
+#pragma endregion
+
 }
 
 #if WITH_EDITOR
@@ -1385,7 +1421,8 @@ void FSteamVRInputDevice::ReloadActionManifest()
 
 		// Set Action Manifest
 		EVRInputError InputError = VRInput()->SetActionManifestPath(TCHAR_TO_UTF8(*IFileManager::Get().ConvertToAbsolutePathForExternalAppForRead(*ManifestPath)));
-		GetInputError(InputError, FString(TEXT("Setting Action Manifest Path")));
+		UE_LOG( LogSteamVRInputDevice, Display, TEXT( "[STEAMVR INPUT] Reloading Action Manifest Path [%s]" ), *ManifestPath );
+		GetInputError(InputError, FString(TEXT("Setting Action Manifest Path to")));
 	}
 }
 #endif
@@ -1791,13 +1828,27 @@ void FSteamVRInputDevice::GenerateActionBindings(TArray<FInputMapping> &InInputM
 				}
 
 				// Set Cache Mode
-				CacheMode = InputState.bIsTrigger || InputState.bIsGrip ? FName(TEXT("trigger")) : FName(TEXT("button"));
-				CacheMode = InputState.bIsPress ? FName(TEXT("button")) : CacheMode;
+				CacheMode = InputState.bIsTrigger ? FName(TEXT("trigger")) : FName(TEXT("button"));
+				CacheMode = InputState.bIsPress && !InputState.bIsTrigger ? FName(TEXT("button")) : CacheMode;
 				CacheMode = InputState.bIsTrackpad ? FName(TEXT("trackpad")) : CacheMode;
 				CacheMode = InputState.bIsThumbstick ? FName(TEXT("joystick")) : CacheMode;
-				CacheMode = InputState.bIsGrip ? FName(TEXT("button")) : CacheMode;
+				
 				CacheMode = InputState.bIsPinchGrab || InputState.bIsGripGrab ? FName(TEXT("grab")) : CacheMode;
 
+				
+				// Set Grip Cache Mode
+				if (CurrentInputKeyName.Contains(TEXT("ValveIndex"), ESearchCase::IgnoreCase, ESearchDir::FromStart)
+					|| CurrentInputKeyName.Contains(TEXT("Oculus"), ESearchCase::IgnoreCase, ESearchDir::FromStart))
+				{
+					// Set Grip Cache Mode for Index & Oculus to "trigger"
+					CacheMode = InputState.bIsGrip ? FName(TEXT("trigger")) : CacheMode;
+				}
+				else
+				{
+					// Any other controller should use the button cache mode
+					CacheMode = InputState.bIsGrip ? FName(TEXT("button")) : CacheMode;
+				}
+
 				// Set Cache Path
 				if (InputState.bIsTrigger)
 				{
@@ -1886,15 +1937,15 @@ void FSteamVRInputDevice::GenerateActionBindings(TArray<FInputMapping> &InInputM
 				// Add parameters if Dpad
 				if (InputState.bIsDpadUp || InputState.bIsDpadDown || InputState.bIsDpadLeft || InputState.bIsDpadRight)
 				{
-					// Create Submode - Vive will use click as default
+					// Create Submode
 					TSharedRef<FJsonObject> SubmodeJsonObject = MakeShareable(new FJsonObject());
-					if (Controller.Description.Contains(TEXT("Vive")))
+					if ( CurrentInputKeyName.Right(5).Equals( TEXT( "Touch" ) ) )
 					{
-						SubmodeJsonObject->SetStringField(TEXT("sub_mode"), TEXT("click"));
+						SubmodeJsonObject->SetStringField( TEXT( "sub_mode" ), TEXT( "touch" ) );
 					}
 					else
 					{
-						SubmodeJsonObject->SetStringField(TEXT("sub_mode"), TEXT("touch"));
+						SubmodeJsonObject->SetStringField( TEXT( "sub_mode" ), TEXT( "click" ) );
 					}
 					
 					// Create Parameter
@@ -2054,6 +2105,12 @@ void FSteamVRInputDevice::GenerateActionBindings(TArray<FInputMapping> &InInputM
 						FString CacheType;
 						FString CachePath;
 	
+						// Reset Dpad States
+						InputState.bIsDpadUp = false;
+						InputState.bIsDpadDown = false;
+						InputState.bIsDpadLeft = false;
+						InputState.bIsDpadRight = false;
+
 						// Set Axis States
 						InputState.bIsAxis = false;
 						InputState.bIsAxis2 = false;
@@ -2105,6 +2162,36 @@ void FSteamVRInputDevice::GenerateActionBindings(TArray<FInputMapping> &InInputM
 								CurrentInputKeyName.Contains(TEXT("_Touch"), ESearchCase::CaseSensitive, ESearchDir::FromEnd);
 						}
 	
+						// Check for DPad Keys
+						if (CurrentInputKeyName.Contains(TEXT("_Up"), ESearchCase::CaseSensitive, ESearchDir::FromEnd))
+						{
+							InputState.bIsDpadUp = true;
+							InputState.bIsDpadDown = false;
+							InputState.bIsDpadLeft = false;
+							InputState.bIsDpadRight = false;
+						}
+						else if (CurrentInputKeyName.Contains(TEXT("_Down"), ESearchCase::CaseSensitive, ESearchDir::FromEnd))
+						{
+							InputState.bIsDpadUp = false;
+							InputState.bIsDpadDown = true;
+							InputState.bIsDpadLeft = false;
+							InputState.bIsDpadRight = false;
+						}
+						else if (CurrentInputKeyName.Contains(TEXT("Trackpad_Left"), ESearchCase::CaseSensitive, ESearchDir::FromEnd) || CurrentInputKeyName.Contains(TEXT("Thumbstick_Left"), ESearchCase::CaseSensitive, ESearchDir::FromEnd))
+						{
+							InputState.bIsDpadUp = false;
+							InputState.bIsDpadDown = false;
+							InputState.bIsDpadLeft = true;
+							InputState.bIsDpadRight = false;
+						}
+						else if (CurrentInputKeyName.Contains(TEXT("Trackpad_Right"), ESearchCase::CaseSensitive, ESearchDir::FromEnd) || CurrentInputKeyName.Contains(TEXT("Thumbstick_Right"), ESearchCase::CaseSensitive, ESearchDir::FromEnd))
+						{
+							InputState.bIsDpadUp = false;
+							InputState.bIsDpadDown = false;
+							InputState.bIsDpadLeft = false;
+							InputState.bIsDpadRight = true;
+						}
+
 						// Handle Special Actions for Knuckles Keys
 						if ((CurrentInputKeyName.Contains(TEXT("ValveIndex"), ESearchCase::IgnoreCase, ESearchDir::FromStart)
 							|| CurrentInputKeyName.Contains(TEXT("Cosmos"), ESearchCase::IgnoreCase, ESearchDir::FromStart)) 
@@ -2147,7 +2234,9 @@ void FSteamVRInputDevice::GenerateActionBindings(TArray<FInputMapping> &InInputM
 	
 						// If key being mapped is not an axis key (hardware-wise), set mode as an analog action (scalar_constant to 1.0f)
 						// [REDACTED_PROFILE_URL]/steamvr_unreal_plugin/issues/12
-						if (!SteamVRAxisKeyMapping.InputAxisKeyMapping.Key.IsFloatAxis())
+						if ( !SteamVRAxisKeyMapping.InputAxisKeyMapping.Key.IsFloatAxis()
+							&& ( !CurrentInputKeyName.Contains( TEXT( "Trackpad" ), ESearchCase::CaseSensitive, ESearchDir::FromEnd )
+								&& !CurrentInputKeyName.Contains( TEXT( "Touch" ), ESearchCase::CaseSensitive, ESearchDir::FromEnd ) ) )
 						{
 							CacheMode = FName(TEXT("scalar_constant"));
 						}
@@ -2215,7 +2304,13 @@ void FSteamVRInputDevice::GenerateActionBindings(TArray<FInputMapping> &InInputM
 						{
 							CachePath = InputState.bIsLeft ? FString(TEXT(ACTION_PATH_GRIP_GRAB_LEFT)) : FString(TEXT(ACTION_PATH_GRIP_GRAB_RIGHT));
 						}
-	
+
+						// Override mode if Dpad
+						if (InputState.bIsDpadUp || InputState.bIsDpadDown || InputState.bIsDpadLeft || InputState.bIsDpadRight)
+						{
+							CacheMode = FName(TEXT("dpad"));
+						}
+
 						// Create Action Source
 						FActionSource ActionSource = FActionSource(CacheMode, CachePath);
 						TSharedRef<FJsonObject> ActionSourceJsonObject = MakeShareable(new FJsonObject());
@@ -2231,6 +2326,29 @@ void FSteamVRInputDevice::GenerateActionBindings(TArray<FInputMapping> &InInputM
 							continue;
 						}
 	
+						// Add parameters if Dpad
+						if (InputState.bIsDpadUp || InputState.bIsDpadDown || InputState.bIsDpadLeft || InputState.bIsDpadRight)
+						{
+							// Set sub mode
+							TSharedRef<FJsonObject> SubmodeJsonObject = MakeShareable(new FJsonObject());
+
+							if ( CurrentInputKeyName.Right( 5 ).Equals( TEXT( "Touch" ) ) )
+							{
+								SubmodeJsonObject->SetStringField( TEXT( "sub_mode" ), TEXT( "touch" ) );
+							}
+							else
+							{
+								SubmodeJsonObject->SetStringField( TEXT( "sub_mode" ), TEXT( "click" ) );
+							}
+
+
+							// Create Parameter
+							TSharedPtr<FJsonObject> ParametersJsonObject = MakeShareable(new FJsonObject());
+
+							// Set Submode as a parameter
+							ActionSourceJsonObject->SetObjectField(TEXT("parameters"), SubmodeJsonObject);
+						}
+
 						// Set Key Mappings
 						TSharedPtr<FJsonObject> ActionInputJsonObject = MakeShareable(new FJsonObject());
 	
@@ -2243,6 +2361,10 @@ void FSteamVRInputDevice::GenerateActionBindings(TArray<FInputMapping> &InInputM
 						{
 							CacheType = FString(TEXT("value"));
 						}
+						else if (CurrentInputKeyName.Contains(TEXT("Trackpad_Touch"), ESearchCase::CaseSensitive, ESearchDir::FromEnd))
+						{
+							CacheType = FString(TEXT("touch"));
+						}
 						else if (InputState.bIsAxis && InputState.bIsAxis2)
 						{
 							if (InputState.bIsGrip)
@@ -2302,6 +2424,24 @@ void FSteamVRInputDevice::GenerateActionBindings(TArray<FInputMapping> &InInputM
 						{
 							CacheType = "";
 						}
+
+						// Handle Dpad values
+						if (InputState.bIsDpadUp)
+						{
+							CacheType = "north";
+						}
+						else if (InputState.bIsDpadDown)
+						{
+							CacheType = "south";
+						}
+						else if (InputState.bIsDpadLeft)
+						{
+							CacheType = "west";
+						}
+						else if (InputState.bIsDpadRight)
+						{
+							CacheType = "east";
+						}
 	
 						// Handle special actions
 						if (InputState.bIsPinchGrab || InputState.bIsGripGrab)
@@ -2445,15 +2585,6 @@ void FSteamVRInputDevice::GenerateActionManifest(bool GenerateActions, bool Gene
 				FName(TEXT("Skeleton (Right)")), FString(TEXT(ACTION_PATH_SKEL_HAND_RIGHT))));
 		}
 
-		// Open console
-		{
-			const FKey* ConsoleKey = InputSettings->ConsoleKeys.FindByPredicate([](FKey& Key) { return Key.IsValid(); });
-			if (ConsoleKey != nullptr)
-			{
-				Actions.Add(FSteamVRInputAction(FString(TEXT(ACTION_PATH_OPEN_CONSOLE)), FName(TEXT("Open Console")), false, ConsoleKey->GetFName(), false));
-			}
-		}
-
 		// Haptics
 		{
 			FString ConstActionPath = FString(TEXT(ACTION_PATH_VIBRATE_LEFT));
@@ -2467,6 +2598,16 @@ void FSteamVRInputDevice::GenerateActionManifest(bool GenerateActions, bool Gene
 		// Add base actions to the action manifest
 		ActionManifestObject->SetArrayField(TEXT("actions"), InputActionsArray);
 
+		// Open console
+		{
+			const FKey* ConsoleKey = InputSettings->ConsoleKeys.FindByPredicate([](FKey& Key) { return Key.IsValid(); });
+			if (ConsoleKey != nullptr)
+			{
+				Actions.Add(FSteamVRInputAction(FString(TEXT(ACTION_PATH_OPEN_CONSOLE)), FName(TEXT("Open Console")), false, ConsoleKey->GetFName(), false));
+				UniqueInputs.AddUnique(ConsoleKey->GetFName());
+			}
+		}
+
 		// Add project's input key mappings to SteamVR's Input Actions
 		ProcessKeyInputMappings(InputSettings, UniqueInputs);
 
@@ -2620,11 +2761,8 @@ void FSteamVRInputDevice::GenerateActionManifest(bool GenerateActions, bool Gene
 		}
 
 		// If there are input actions, add them to the action manifest object
-		if (InputActionsArray.Num() > 0)
-		{
-			ActionManifestObject->SetArrayField(TEXT("actions"), InputActionsArray);
-		}
-	}
+		ActionManifestObject->SetArrayField(TEXT("actions"), InputActionsArray);
+	}	
 	else
 	{
 		UE_LOG(LogSteamVRInputDevice, Error, TEXT("Error trying to retrieve Input Settings."));
