Enable SteamVR Input plugins for Mac, Linux, and Win32 (otherwise there is content loading errors because types from the plugin are not defined).
Fix various deprecations and code build errors on clang.

diff --git a/Engine/Plugins/Runtime/Steam/SteamVR/Source/SteamVREditor/Private/SteamVREditorCommands.cpp b/Engine/Plugins/Runtime/Steam/SteamVR/Source/SteamVREditor/Private/SteamVREditorCommands.cpp
--- a/Engine/Plugins/Runtime/Steam/SteamVR/Source/SteamVREditor/Private/SteamVREditorCommands.cpp
+++ b/Engine/Plugins/Runtime/Steam/SteamVR/Source/SteamVREditor/Private/SteamVREditorCommands.cpp
@@ -35,12 +35,12 @@ POSSIBILITY OF SUCH DAMAGE.
 
 void FSteamVREditorCommands::RegisterCommands()
 {
-	UI_COMMAND(PluginAction, "SteamVREditor", "Execute SteamVREditor action", EUserInterfaceActionType::Button, FInputGesture());
-	UI_COMMAND(JsonActionManifest, "Regenerate Action Manifest", "Regenerate Action Manifest", EUserInterfaceActionType::Button, FInputGesture());
-	UI_COMMAND(JsonControllerBindings, "Regenerate Controller Bindings", "Regenerate Controller Bindings", EUserInterfaceActionType::Button, FInputGesture());
-	UI_COMMAND(ReloadActionManifest, "Reload Action Manifest", "Reload Action Manifest", EUserInterfaceActionType::Button, FInputGesture());
-	UI_COMMAND(LaunchBindingsURL, "Launch SteamVR Bindings Dashboard", "Launch SteamVR Bindings Dashboard", EUserInterfaceActionType::Button, FInputGesture());
-	UI_COMMAND(AddSampleInputs, "Add Sample Inputs", "Add Sample Inputs", EUserInterfaceActionType::Button, FInputGesture());
+	UI_COMMAND(PluginAction, "SteamVREditor", "Execute SteamVREditor action", EUserInterfaceActionType::Button, FInputChord());
+	UI_COMMAND(JsonActionManifest, "Regenerate Action Manifest", "Regenerate Action Manifest", EUserInterfaceActionType::Button, FInputChord());
+	UI_COMMAND(JsonControllerBindings, "Regenerate Controller Bindings", "Regenerate Controller Bindings", EUserInterfaceActionType::Button, FInputChord());
+	UI_COMMAND(ReloadActionManifest, "Reload Action Manifest", "Reload Action Manifest", EUserInterfaceActionType::Button, FInputChord());
+	UI_COMMAND(LaunchBindingsURL, "Launch SteamVR Bindings Dashboard", "Launch SteamVR Bindings Dashboard", EUserInterfaceActionType::Button, FInputChord());
+	UI_COMMAND(AddSampleInputs, "Add Sample Inputs", "Add Sample Inputs", EUserInterfaceActionType::Button, FInputChord());
 }
 
 #undef LOCTEXT_NAMESPACE
diff --git a/Engine/Plugins/Runtime/Steam/SteamVR/Source/SteamVRInputDevice/Private/SteamVRInputDevice.cpp b/Engine/Plugins/Runtime/Steam/SteamVR/Source/SteamVRInputDevice/Private/SteamVRInputDevice.cpp
--- a/Engine/Plugins/Runtime/Steam/SteamVR/Source/SteamVRInputDevice/Private/SteamVRInputDevice.cpp
+++ b/Engine/Plugins/Runtime/Steam/SteamVR/Source/SteamVRInputDevice/Private/SteamVRInputDevice.cpp
@@ -2522,77 +2522,77 @@ void FSteamVRInputDevice::GenerateActionManifest(bool GenerateActions, bool Gene
 		// Controller poses
 		{
 			FString ConstActionPath = FString(TEXT(ACTION_PATH_CONTROLLER_LEFT));
-			Actions.Add(FSteamVRInputAction(ConstActionPath, EActionType::Pose, false,
+			Actions.Add(FSteamVRInputAction(ConstActionPath, ESteamVRActionType::Pose, false,
 				FName(TEXT("Left Controller [Pose]")), FString(TEXT(ACTION_PATH_CONT_RAW_LEFT))));
 		}
 		{
 			FString ConstActionPath = FString(TEXT(ACTION_PATH_CONTROLLER_RIGHT));
-			Actions.Add(FSteamVRInputAction(ConstActionPath, EActionType::Pose, false,
+			Actions.Add(FSteamVRInputAction(ConstActionPath, ESteamVRActionType::Pose, false,
 				FName(TEXT("Right Controller [Pose]")), FString(TEXT(ACTION_PATH_CONT_RAW_RIGHT))));
 		}
 
 		// Other poses
 		{
 			FString ConstActionPath = FString(TEXT(ACTION_PATH_SPECIAL_BACK_L));
-			Actions.Add(FSteamVRInputAction(ConstActionPath, EActionType::Pose, false,
+			Actions.Add(FSteamVRInputAction(ConstActionPath, ESteamVRActionType::Pose, false,
 				FName(TEXT("Special 1 [Tracker]")), FString(TEXT(ACTION_PATH_SPCL_BACK_LEFT))));
 		}
 		{
 			FString ConstActionPath = FString(TEXT(ACTION_PATH_SPECIAL_BACK_R));
-			Actions.Add(FSteamVRInputAction(ConstActionPath, EActionType::Pose, false,
+			Actions.Add(FSteamVRInputAction(ConstActionPath, ESteamVRActionType::Pose, false,
 				FName(TEXT("Special 2 [Tracker]")), FString(TEXT(ACTION_PATH_SPCL_BACK_RIGHT))));
 		}
 		{
 			FString ConstActionPath = FString(TEXT(ACTION_PATH_SPECIAL_FRONT_L));
-			Actions.Add(FSteamVRInputAction(ConstActionPath, EActionType::Pose, false,
+			Actions.Add(FSteamVRInputAction(ConstActionPath, ESteamVRActionType::Pose, false,
 				FName(TEXT("Special 3 [Tracker]")), FString(TEXT(ACTION_PATH_SPCL_FRONT_LEFT))));
 		}
 		{
 			FString ConstActionPath = FString(TEXT(ACTION_PATH_SPECIAL_FRONT_R));
-			Actions.Add(FSteamVRInputAction(ConstActionPath, EActionType::Pose, false,
+			Actions.Add(FSteamVRInputAction(ConstActionPath, ESteamVRActionType::Pose, false,
 				FName(TEXT("Special 4 [Tracker]")), FString(TEXT(ACTION_PATH_SPCL_FRONT_RIGHT))));
 		}
 		{
 			FString ConstActionPath = FString(TEXT(ACTION_PATH_SPECIAL_FRONTR_L));
-			Actions.Add(FSteamVRInputAction(ConstActionPath, EActionType::Pose, false,
+			Actions.Add(FSteamVRInputAction(ConstActionPath, ESteamVRActionType::Pose, false,
 				FName(TEXT("Special 5 [Tracker]")), FString(TEXT(ACTION_PATH_SPCL_FRONTR_LEFT))));
 		}
 		{
 			FString ConstActionPath = FString(TEXT(ACTION_PATH_SPECIAL_FRONTR_R));
-			Actions.Add(FSteamVRInputAction(ConstActionPath, EActionType::Pose, false,
+			Actions.Add(FSteamVRInputAction(ConstActionPath, ESteamVRActionType::Pose, false,
 				FName(TEXT("Special 6 [Tracker]")), FString(TEXT(ACTION_PATH_SPCL_FRONTR_RIGHT))));
 		}
 		{
 			FString ConstActionPath = FString(TEXT(ACTION_PATH_SPECIAL_PISTOL_L));
-			Actions.Add(FSteamVRInputAction(ConstActionPath, EActionType::Pose, false,
+			Actions.Add(FSteamVRInputAction(ConstActionPath, ESteamVRActionType::Pose, false,
 				FName(TEXT("Special 7 [Tracker]")), FString(TEXT(ACTION_PATH_SPCL_PISTOL_LEFT))));
 		}
 		{
 			FString ConstActionPath = FString(TEXT(ACTION_PATH_SPECIAL_PISTOL_R));
-			Actions.Add(FSteamVRInputAction(ConstActionPath, EActionType::Pose, false,
+			Actions.Add(FSteamVRInputAction(ConstActionPath, ESteamVRActionType::Pose, false,
 				FName(TEXT("Special 8 [Tracker]")), FString(TEXT(ACTION_PATH_SPCL_PISTOL_RIGHT))));
 		}
 
 		// Skeletal Data
 		{
 			FString ConstActionPath = FString(TEXT(ACTION_PATH_SKELETON_LEFT));
-			Actions.Add(FSteamVRInputAction(ConstActionPath, EActionType::Skeleton, false,
+			Actions.Add(FSteamVRInputAction(ConstActionPath, ESteamVRActionType::Skeleton, false,
 				FName(TEXT("Skeleton (Left)")), FString(TEXT(ACTION_PATH_SKEL_HAND_LEFT))));
 		}
 		{
 			FString ConstActionPath = FString(TEXT(ACTION_PATH_SKELETON_RIGHT));
-			Actions.Add(FSteamVRInputAction(ConstActionPath, EActionType::Skeleton, false,
+			Actions.Add(FSteamVRInputAction(ConstActionPath, ESteamVRActionType::Skeleton, false,
 				FName(TEXT("Skeleton (Right)")), FString(TEXT(ACTION_PATH_SKEL_HAND_RIGHT))));
 		}
 
 		// Haptics
 		{
 			FString ConstActionPath = FString(TEXT(ACTION_PATH_VIBRATE_LEFT));
-			Actions.Add(FSteamVRInputAction(ConstActionPath, EActionType::Vibration, false, FName(TEXT("Haptic (Left)"))));
+			Actions.Add(FSteamVRInputAction(ConstActionPath, ESteamVRActionType::Vibration, false, FName(TEXT("Haptic (Left)"))));
 		}
 		{
 			FString ConstActionPath = FString(TEXT(ACTION_PATH_VIBRATE_RIGHT));
-			Actions.Add(FSteamVRInputAction(ConstActionPath, EActionType::Vibration, false, FName(TEXT("Haptic (Right)"))));
+			Actions.Add(FSteamVRInputAction(ConstActionPath, ESteamVRActionType::Vibration, false, FName(TEXT("Haptic (Right)"))));
 		}
 
 		// Add base actions to the action manifest
@@ -2629,7 +2629,7 @@ void FSteamVRInputDevice::GenerateActionManifest(bool GenerateActions, bool Gene
 				//UE_LOG(LogTemp, Warning, TEXT("Action: %s Type: [%s]"), *Action.KeyX.ToString(), *Action.GetActionTypeName());
 
 				// Check for boolean/digital input
-				if (Action.Type == EActionType::Boolean)
+				if (Action.Type == ESteamVRActionType::Boolean)
 				{
 					// Set Key Actions Linked To This Input Key
 					TArray<FInputActionKeyMapping> ActionKeyMappings;
@@ -2644,7 +2644,7 @@ void FSteamVRInputDevice::GenerateActionManifest(bool GenerateActions, bool Gene
 				}
 
 				// Check for axes/analog input
-				if (Action.Type == EActionType::Vector1 || Action.Type == EActionType::Vector2 || Action.Type == EActionType::Vector3)
+				if (Action.Type == ESteamVRActionType::Vector1 || Action.Type == ESteamVRActionType::Vector2 || Action.Type == ESteamVRActionType::Vector3)
 				{
 					// Set Axis Actions Linked To This Input Key
 					FString ActionAxis = Action.Name.ToString();
@@ -2696,7 +2696,7 @@ void FSteamVRInputDevice::GenerateActionManifest(bool GenerateActions, bool Gene
 				};
 
 				// Add hand if skeleton
-				if (Action.Type == EActionType::Skeleton)
+				if (Action.Type == ESteamVRActionType::Skeleton)
 				{
 					ActionFields.Append({ TEXT("skeleton"), Action.StringPath });
 				}
@@ -3719,7 +3719,7 @@ void FSteamVRInputDevice::ProcessActionEvents(FSteamVRInputActionSet SteamVRInpu
 			continue;
 		}
 
-		if (Action.Type == EActionType::Boolean && !Action.Path.Contains(TEXT(" axis")) && !Action.Path.Contains(TEXT("_axis")))
+		if (Action.Type == ESteamVRActionType::Boolean && !Action.Path.Contains(TEXT(" axis")) && !Action.Path.Contains(TEXT("_axis")))
 		{
 			// Get digital data from SteamVR
 			InputDigitalActionData_t DigitalData;
@@ -3776,9 +3776,9 @@ void FSteamVRInputDevice::ProcessActionEvents(FSteamVRInputActionSet SteamVRInpu
 				}
 			}
 		}
-		else if (Action.Type == EActionType::Vector1
-			|| Action.Type == EActionType::Vector2
-			|| Action.Type == EActionType::Vector3)
+		else if (Action.Type == ESteamVRActionType::Vector1
+			|| Action.Type == ESteamVRActionType::Vector2
+			|| Action.Type == ESteamVRActionType::Vector3)
 		{
 			// Get analog data from SteamVR
 			InputAnalogActionData_t AnalogData;
diff --git a/Engine/Plugins/Runtime/Steam/SteamVR/Source/SteamVRInputDevice/Private/SteamVRInputTypes.h b/Engine/Plugins/Runtime/Steam/SteamVR/Source/SteamVRInputDevice/Private/SteamVRInputTypes.h
--- a/Engine/Plugins/Runtime/Steam/SteamVR/Source/SteamVRInputDevice/Private/SteamVRInputTypes.h
+++ b/Engine/Plugins/Runtime/Steam/SteamVR/Source/SteamVRInputDevice/Private/SteamVRInputTypes.h
@@ -165,7 +165,7 @@ struct ESteamVRInputButton
 };
 
 /* Available SteamVR Input Action Types */
-enum EActionType
+enum class ESteamVRActionType
 {
 	Boolean,
 	Vector1,
@@ -313,7 +313,7 @@ struct FSteamVRInputAction
 {
 	FString		Path;							// The path defined for the action (e.g. main/in/{ActionName})
 	FName		Name;							// The SteamVR name of the action (e.g. Teleport, OpenConsole)
-	EActionType	Type;							// The SteamVR input data type
+	ESteamVRActionType	Type;							// The SteamVR input data type
 	FName		KeyX;							// The UE Key in the X axis or float axis (e.g. Motion_Controller_Thumbstick_X)
 	FName		KeyY;							// The UE Key in the Y axis
 	FName		KeyZ;							// The UE Key in the Z axis
@@ -334,7 +334,7 @@ struct FSteamVRInputAction
 		return SActionTypes[(int)Type];
 	}
 
-	FSteamVRInputAction(const FString& inPath, EActionType inType, bool inRequirement, const FName& inName, const FString& inStringPath)
+	FSteamVRInputAction(const FString& inPath, ESteamVRActionType inType, bool inRequirement, const FName& inName, const FString& inStringPath)
 		: Path(inPath)
 		, Name(inName)
 		, Type(inType)
@@ -354,7 +354,7 @@ struct FSteamVRInputAction
 	FSteamVRInputAction(const FString& inPath, const FName& inName, const FName& inKeyName, bool inState)
 		: Path(inPath)
 		, Name(inName)
-		, Type(Boolean)
+		, Type(ESteamVRActionType::Boolean)
 		, KeyX(inKeyName)
 		, KeyY()
 		, KeyZ()
@@ -371,7 +371,7 @@ struct FSteamVRInputAction
 	FSteamVRInputAction(const FString& inPath, const FName& inName, bool inRequirement, const FName& inKeyName, bool inState)
 		: Path(inPath)
 		, Name(inName)
-		, Type(Boolean)
+		, Type(ESteamVRActionType::Boolean)
 		, KeyX(inKeyName)
 		, KeyY()
 		, KeyZ()
@@ -388,7 +388,7 @@ struct FSteamVRInputAction
 	FSteamVRInputAction(const FString& inPath, const FName& inName, const FName& inKeyName, float inValue1D)
 		: Path(inPath)
 		, Name(inName)
-		, Type(Vector1)
+		, Type(ESteamVRActionType::Vector1)
 		, KeyX(inKeyName)
 		, KeyY()
 		, KeyZ()
@@ -401,7 +401,7 @@ struct FSteamVRInputAction
 	FSteamVRInputAction(const FString& inPath, const FName& inName, const FName& inKeyName_X, const FName& inKeyName_Y, const FVector2D& inValue2D)
 		: Path(inPath)
 		, Name(inName)
-		, Type(Vector2)
+		, Type(ESteamVRActionType::Vector2)
 		, KeyX(inKeyName_X)
 		, KeyY(inKeyName_Y)
 		, KeyZ()
@@ -417,7 +417,7 @@ struct FSteamVRInputAction
 	FSteamVRInputAction(const FString& inPath, const FName& inName, const FName& inKeyName_X, const FName& inKeyName_Y, const FName& inKeyName_Z, const FVector& inValue3D)
 		: Path(inPath)
 		, Name(inName)
-		, Type(Vector3)
+		, Type(ESteamVRActionType::Vector3)
 		, KeyX(inKeyName_X)
 		, KeyY(inKeyName_Y)
 		, KeyZ(inKeyName_Z)
@@ -430,7 +430,7 @@ struct FSteamVRInputAction
 		LastUpdated = 0.f;
 	}
 
-	FSteamVRInputAction(const FString& inPath, const EActionType& inActionType, const bool& inRequirement, const FName& inName)
+	FSteamVRInputAction(const FString& inPath, const ESteamVRActionType& inActionType, const bool& inRequirement, const FName& inName)
 		: Path(inPath)
 		, Name(inName)
 		, Type(inActionType)
diff --git a/Engine/Plugins/Runtime/Steam/SteamVR/Source/SteamVRInputDevice/Public/ISteamVRInputDeviceModule.h b/Engine/Plugins/Runtime/Steam/SteamVR/Source/SteamVRInputDevice/Public/ISteamVRInputDeviceModule.h
--- a/Engine/Plugins/Runtime/Steam/SteamVR/Source/SteamVRInputDevice/Public/ISteamVRInputDeviceModule.h
+++ b/Engine/Plugins/Runtime/Steam/SteamVR/Source/SteamVRInputDevice/Public/ISteamVRInputDeviceModule.h
@@ -34,7 +34,7 @@ POSSIBILITY OF SUCH DAMAGE.
 #include "Modules/ModuleManager.h"
 #include "IInputDeviceModule.h"
 
-STEAMVRINPUT_API class ISteamVRInputDeviceModule : public IInputDeviceModule
+class STEAMVRINPUTDEVICE_API ISteamVRInputDeviceModule : public IInputDeviceModule
 {
 
 public:
diff --git a/Engine/Plugins/Runtime/Steam/SteamVR/SteamVR.uplugin b/Engine/Plugins/Runtime/Steam/SteamVR/SteamVR.uplugin
--- a/Engine/Plugins/Runtime/Steam/SteamVR/SteamVR.uplugin
+++ b/Engine/Plugins/Runtime/Steam/SteamVR/SteamVR.uplugin
@@ -37,7 +37,10 @@
 			"Type": "Runtime",
 			"LoadingPhase": "PostConfigInit",
 			"WhitelistPlatforms": [
-				"Win64"
+				"Win64",
+				"Win32",
+				"Linux",
+				"Mac"
 			]
 		},
 		{
@@ -45,7 +48,10 @@
 			"Type": "Runtime",
 			"LoadingPhase": "PostConfigInit",
 			"WhitelistPlatforms": [
-				"Win64"
+				"Win64",
+				"Win32",
+				"Linux",
+				"Mac"
 			]
 		},
 		{
@@ -53,7 +59,10 @@
 			"Type": "Editor",
 			"LoadingPhase": "PostEngineInit",
 			"WhitelistPlatforms": [
-				"Win64"
+				"Win64",
+				"Win32",
+				"Linux",
+				"Mac"
 			]
 		}
 	],
