Adding disabled by default preference for showing SteamVR Input button in the Toolbar, controlled by new SteamVR Editor settings in the Project Settings.  Duplicated the buttons under there as well.

Also adding the active XR system name to the window title if doing VRPreview with XR System enabled, as a better alternative to tell which system is the active one, that works for all systems and not just SteamVR

diff --git a/Engine/Plugins/Runtime/Steam/SteamVR/Source/SteamVREditor/Private/SteamVREditor.cpp b/Engine/Plugins/Runtime/Steam/SteamVR/Source/SteamVREditor/Private/SteamVREditor.cpp
--- a/Engine/Plugins/Runtime/Steam/SteamVR/Source/SteamVREditor/Private/SteamVREditor.cpp
+++ b/Engine/Plugins/Runtime/Steam/SteamVR/Source/SteamVREditor/Private/SteamVREditor.cpp
@@ -36,6 +36,14 @@ POSSIBILITY OF SUCH DAMAGE.
 #include "LevelEditor.h"
 #include "Engine/Engine.h"
 #include "IXRTrackingSystem.h"
+#include "ISettingsModule.h"
+#include "SteamVREditorSettings.h"
+#include "DetailCategoryBuilder.h"
+#include "DetailLayoutBuilder.h"
+#include "Widgets/Input/SButton.h"
+#include "DetailWidgetRow.h"
+#include "Runtime/Core/Public/CoreGlobals.h"
+#include "Runtime/Core/Public/Misc/ConfigCacheIni.h"
 
 static const FName SteamVREditorTabName("SteamVREditor");
 
@@ -43,6 +51,8 @@ static const FName SteamVREditorTabName("SteamVREditor");
 
 void FSteamVREditorModule::StartupModule()
 {
+	RegisterSettings();
+
 	FSteamVREditorStyle::Initialize();
 	FSteamVREditorStyle::ReloadTextures();
 
@@ -309,16 +319,27 @@ void FSteamVREditorModule::AddMenuExtension(FMenuBuilder& Builder)
 
 void FSteamVREditorModule::AddToolbarExtension(FToolBarBuilder& Builder)
 {
-	FSteamVREditorStyle MenuStyle = FSteamVREditorStyle();
-	MenuStyle.Initialize();
-
-	Builder.AddComboButton(
-		FUIAction(FExecuteAction::CreateRaw(this, &FSteamVREditorModule::PluginButtonClicked)),
-		FOnGetContent::CreateRaw(this, &FSteamVREditorModule::FillComboButton, PluginCommands),
-		LOCTEXT("SteamVRInputBtn", "SteamVR Input"),
-		LOCTEXT("SteamVRInputBtnTootlip", "SteamVR Input"),
-		FSlateIcon(FSteamVREditorStyle::GetStyleSetName(), "SteamVREditor.PluginAction")
-	);
+	// Only show the toolbar button if enabled by the user for this project.  Setting is under Project Settings > PLugins > SteamVR now instead.
+	bool bShowSteamVrInputToolbarButton;
+	if (GConfig->GetBool(TEXT("/Script/SteamVREditor.SteamVREditorSettings"), TEXT("bShowSteamVrInputToolbarButton"), bShowSteamVrInputToolbarButton, GEditorIni) && bShowSteamVrInputToolbarButton)
+	{
+		FSteamVREditorStyle MenuStyle = FSteamVREditorStyle();
+		MenuStyle.Initialize();
+
+		static const FName SteamVRInputSectionName = FName("SteamVR Input");
+
+		Builder.BeginSection(SteamVRInputSectionName);
+
+		Builder.AddComboButton(
+			FUIAction(FExecuteAction::CreateRaw(this, &FSteamVREditorModule::PluginButtonClicked)),
+			FOnGetContent::CreateRaw(this, &FSteamVREditorModule::FillComboButton, PluginCommands),
+			LOCTEXT("SteamVRInputBtn", "SteamVR Input"),
+			LOCTEXT("SteamVRInputBtnTootlip", "SteamVR Input"),
+			FSlateIcon(FSteamVREditorStyle::GetStyleSetName(), "SteamVREditor.PluginAction")
+		);
+
+		Builder.EndSection();
+	}
 }
 
 TSharedRef<SWidget> FSteamVREditorModule::FillComboButton(TSharedPtr<class FUICommandList> Commands)
@@ -334,6 +355,121 @@ TSharedRef<SWidget> FSteamVREditorModule::FillComboButton(TSharedPtr<class FUICo
 	return MenuBuilder.MakeWidget();
 }
 
-#undef LOCTEXT_NAMESPACE
+void FSteamVREditorModule::RegisterSettings()
+{
+	if (ISettingsModule* SettingsModule = FModuleManager::GetModulePtr<ISettingsModule>("Settings"))
+	{
+		SettingsModule->RegisterSettings("Project", "Plugins", "SteamVR",
+			LOCTEXT("SteamVREditorSettingsName", "SteamVR"),
+			LOCTEXT("SteamVREditorSettingsDescription", "Configure the SteamVR plugin"),
+			GetMutableDefault<USteamVREditorSettings>()
+		);
+
+		FPropertyEditorModule& PropertyModule = FModuleManager::GetModuleChecked<FPropertyEditorModule>("PropertyEditor");
+		PropertyModule.RegisterCustomClassLayout(USteamVREditorSettings::StaticClass()->GetFName(), FOnGetDetailCustomizationInstance::CreateStatic(&FSteamVRSettingsDetailsCustomization::MakeInstance));
+	}
+}
+
+void FSteamVREditorModule::UnregisterSettings()
+{
+	if (ISettingsModule* SettingsModule = FModuleManager::GetModulePtr<ISettingsModule>("Settings"))
+	{
+		SettingsModule->UnregisterSettings("Project", "Plugins", "SteamVR");
+	}
+}
+
 	
-IMPLEMENT_MODULE(FSteamVREditorModule, SteamVREditor)
\ No newline at end of file
+IMPLEMENT_MODULE(FSteamVREditorModule, SteamVREditor)
+
+TSharedRef<IDetailCustomization> FSteamVRSettingsDetailsCustomization::MakeInstance()
+{
+	return MakeShareable(new FSteamVRSettingsDetailsCustomization);
+}
+
+void FSteamVRSettingsDetailsCustomization::CustomizeDetails(IDetailLayoutBuilder& DetailLayout)
+{
+	
+	IDetailCategoryBuilder& CategoryBuilder = DetailLayout.EditCategory("SteamVR Input", FText::GetEmpty());
+	CategoryBuilder.AddCustomRow(LOCTEXT("SteamVR Input Category", "SteamVR Input"))
+		.WholeRowContent()
+		[
+			SNew(SVerticalBox)
+			+ SVerticalBox::Slot().AutoHeight().Padding(2)
+		[
+			SNew(SHorizontalBox)
+			+ SHorizontalBox::Slot().AutoWidth()
+		[
+			SNew(SButton)
+			.Text(LOCTEXT("Regenerate Action Manifest", "Regenerate Action Manifest"))
+		.OnClicked_Raw(this, &FSteamVRSettingsDetailsCustomization::RegenActionManifest)
+		]
+	+ SHorizontalBox::Slot().FillWidth(8)
+		]
+	+ SVerticalBox::Slot().AutoHeight().Padding(2)
+		[
+			SNew(SHorizontalBox)
+			+ SHorizontalBox::Slot().AutoWidth()
+		[
+			SNew(SButton)
+			.Text(LOCTEXT("Regenerate Controller Bindings", "Regenerate Controller Bindings"))
+		.OnClicked_Raw(this, &FSteamVRSettingsDetailsCustomization::RegenerateControllerBindings)
+		]
+	+ SHorizontalBox::Slot().FillWidth(8)
+		]
+
+	+ SVerticalBox::Slot().AutoHeight().Padding(2)
+		[
+			SNew(SHorizontalBox)
+			+ SHorizontalBox::Slot().AutoWidth()
+		[
+			SNew(SButton)
+			.Text(LOCTEXT("Reload Action Manifest", "Reload Action Manifest"))
+		.OnClicked_Raw(this, &FSteamVRSettingsDetailsCustomization::ReloadActionManifest)
+		]
+	+ SHorizontalBox::Slot().FillWidth(8)
+		]
+
+	+ SVerticalBox::Slot().AutoHeight().Padding(2)
+		[
+			SNew(SHorizontalBox)
+			+ SHorizontalBox::Slot().AutoWidth()
+		[
+			SNew(SButton)
+			.Text(LOCTEXT("Launch SteamVR Bindings Dashboard", "Launch SteamVR Bindings Dashboard"))
+		.OnClicked_Raw(this, &FSteamVRSettingsDetailsCustomization::LaunchBindingsURL)
+		]
+	+ SHorizontalBox::Slot().FillWidth(8)
+		]
+
+		];
+}
+
+FReply FSteamVRSettingsDetailsCustomization::RegenActionManifest()
+{
+	FSteamVREditorModule& SteamVREditorModule = FModuleManager::LoadModuleChecked<FSteamVREditorModule>("SteamVREditor");
+	SteamVREditorModule.JsonRegenerateActionManifest();
+	return FReply::Handled();
+}
+
+FReply FSteamVRSettingsDetailsCustomization::RegenerateControllerBindings()
+{
+	FSteamVREditorModule& SteamVREditorModule = FModuleManager::LoadModuleChecked<FSteamVREditorModule>("SteamVREditor");
+	SteamVREditorModule.JsonRegenerateControllerBindings();
+	return FReply::Handled();
+}
+
+FReply FSteamVRSettingsDetailsCustomization::ReloadActionManifest()
+{
+	FSteamVREditorModule& SteamVREditorModule = FModuleManager::LoadModuleChecked<FSteamVREditorModule>("SteamVREditor");
+	SteamVREditorModule.ReloadActionManifest();
+	return FReply::Handled();
+}
+
+FReply FSteamVRSettingsDetailsCustomization::LaunchBindingsURL()
+{
+	FSteamVREditorModule& SteamVREditorModule = FModuleManager::LoadModuleChecked<FSteamVREditorModule>("SteamVREditor");
+	SteamVREditorModule.LaunchBindingsURL();
+	return FReply::Handled();
+}
+
+#undef LOCTEXT_NAMESPACE
\ No newline at end of file
diff --git a/Engine/Plugins/Runtime/Steam/SteamVR/Source/SteamVREditor/Private/SteamVREditorSettings.cpp b/Engine/Plugins/Runtime/Steam/SteamVR/Source/SteamVREditor/Private/SteamVREditorSettings.cpp
new file mode 100644
--- /dev/null
+++ b/Engine/Plugins/Runtime/Steam/SteamVR/Source/SteamVREditor/Private/SteamVREditorSettings.cpp
@@ -0,0 +1,5 @@
+// Copyright Epic Games, Inc. All Rights Reserved.
+
+
+#include "SteamVREditorSettings.h"
+
diff --git a/Engine/Plugins/Runtime/Steam/SteamVR/Source/SteamVREditor/Public/SteamVREditor.h b/Engine/Plugins/Runtime/Steam/SteamVR/Source/SteamVREditor/Public/SteamVREditor.h
--- a/Engine/Plugins/Runtime/Steam/SteamVR/Source/SteamVREditor/Public/SteamVREditor.h
+++ b/Engine/Plugins/Runtime/Steam/SteamVR/Source/SteamVREditor/Public/SteamVREditor.h
@@ -32,6 +32,7 @@ POSSIBILITY OF SUCH DAMAGE.
 
 #include "Modules/ModuleManager.h"
 #include "SteamVRInputDeviceFunctionLibrary.h"
+#include "IDetailCustomization.h"
 
 class FToolBarBuilder;
 class FMenuBuilder;
@@ -54,6 +55,9 @@ public:
 	void LaunchBindingsURL();
 	void AddSampleInputs();
 
+	void RegisterSettings();
+	void UnregisterSettings();
+
 private:
 
 	void AddToolbarExtension(FToolBarBuilder& Builder);
@@ -65,3 +69,19 @@ private:
 	TSharedRef<SWidget> FillComboButton(TSharedPtr<class FUICommandList> Commands);
 
 };
+
+class FSteamVRSettingsDetailsCustomization : public IDetailCustomization
+{
+public:
+	/** Makes a new instance of this detail layout class for a specific detail view requesting it */
+	static TSharedRef<IDetailCustomization> MakeInstance();
+
+	// IDetailCustomization interface
+	virtual void CustomizeDetails(IDetailLayoutBuilder& DetailLayout) override;
+	// End of IDetailCustomization interface
+
+	FReply RegenActionManifest();
+	FReply RegenerateControllerBindings();
+	FReply ReloadActionManifest();
+	FReply LaunchBindingsURL();
+};
\ No newline at end of file
diff --git a/Engine/Plugins/Runtime/Steam/SteamVR/Source/SteamVREditor/Public/SteamVREditorSettings.h b/Engine/Plugins/Runtime/Steam/SteamVR/Source/SteamVREditor/Public/SteamVREditorSettings.h
new file mode 100644
--- /dev/null
+++ b/Engine/Plugins/Runtime/Steam/SteamVR/Source/SteamVREditor/Public/SteamVREditorSettings.h
@@ -0,0 +1,20 @@
+// Copyright Epic Games, Inc. All Rights Reserved.
+
+#pragma once
+
+#include "CoreMinimal.h"
+#include "UObject/NoExportTypes.h"
+#include "SteamVREditorSettings.generated.h"
+
+/**
+ * 
+ */
+UCLASS(config = Editor, defaultconfig)
+class STEAMVREDITOR_API USteamVREditorSettings : public UObject
+{
+	GENERATED_BODY()
+public:
+	/** Whether or not to show the SteamVR Input settings toolbar button */
+	UPROPERTY(config, EditAnywhere, Category = "SteamVR Editor Settings")
+	bool bShowSteamVrInputToolbarButton;
+};
diff --git a/Engine/Source/Editor/UnrealEd/Private/PlayLevel.cpp b/Engine/Source/Editor/UnrealEd/Private/PlayLevel.cpp
--- a/Engine/Source/Editor/UnrealEd/Private/PlayLevel.cpp
+++ b/Engine/Source/Editor/UnrealEd/Private/PlayLevel.cpp
@@ -121,7 +121,7 @@ DEFINE_LOG_CATEGORY(LogPlayLevel);
 const static FName NAME_CategoryPIE("PIE");
 
 // Forward declare local utility functions
-FText GeneratePIEViewportWindowTitle(const EPlayNetMode InNetMode, const ERHIFeatureLevel::Type InFeatureLevel, const int32 ClientIndex = -1);
+FText GeneratePIEViewportWindowTitle(const EPlayNetMode InNetMode, const ERHIFeatureLevel::Type InFeatureLevel, const FRequestPlaySessionParams& InSessionParams, const int32 ClientIndex = -1);
 bool PromptMatineeClose();
 
 
@@ -2814,7 +2814,7 @@ UGameInstance* UEditorEngine::CreateInnerProcessPIEGameInstance(FRequestPlaySess
 	return GameInstance;
 }
 
-FText GeneratePIEViewportWindowTitle(const EPlayNetMode InNetMode, const ERHIFeatureLevel::Type InFeatureLevel, const int32 ClientIndex)
+FText GeneratePIEViewportWindowTitle(const EPlayNetMode InNetMode, const ERHIFeatureLevel::Type InFeatureLevel, const FRequestPlaySessionParams& InSessionParams, const int32 ClientIndex)
 {
 #if PLATFORM_64BITS
 	const FString PlatformBitsString(TEXT("64"));
@@ -2842,7 +2842,17 @@ FText GeneratePIEViewportWindowTitle(const EPlayNetMode InNetMode, const ERHIFea
 		Args.Add(TEXT("NetMode"), FText::FromString(TEXT("Standalone")));
 	}
 
-	return FText::Format(NSLOCTEXT("UnrealEd", "PlayInEditor_WindowTitleFormat", "{GameName} Preview [NetMode: {NetMode}] ({PlatformBits}-bit/{RHIName})"), Args);
+	if (GEngine->StereoRenderingDevice && GEngine->StereoRenderingDevice.IsValid() && InSessionParams.SessionPreviewTypeOverride == EPlaySessionPreviewType::VRPreview &&
+		GEngine->XRSystem && GEngine->XRSystem.IsValid())
+	{
+		Args.Add(TEXT("XRSystemName"), FText::FromName(GEngine->XRSystem->GetSystemName()));
+	}
+	else
+	{
+		Args.Add(TEXT("XRSystemName"), FText::GetEmpty());
+	}
+
+	return FText::Format(NSLOCTEXT("UnrealEd", "PlayInEditor_WindowTitleFormat", "{GameName} Preview [NetMode: {NetMode}] ({PlatformBits}-bit/{RHIName}) {XRSystemName}"), Args);
 }
 
 void UEditorEngine::TransferEditorSelectionToSIEInstances()
@@ -2909,7 +2919,7 @@ TSharedRef<SPIEViewport> UEditorEngine::GeneratePIEViewportWindow(const FRequest
 	if (!bHasCustomWindow)
 	{
 
-		FText ViewportName = GeneratePIEViewportWindowTitle(InNetMode, PreviewPlatform.GetEffectivePreviewFeatureLevel(), InViewportIndex);
+		FText ViewportName = GeneratePIEViewportWindowTitle(InNetMode, PreviewPlatform.GetEffectivePreviewFeatureLevel(), InSessionParams, InViewportIndex);
 		PieWindow = SNew(SWindow)
 			.Title(ViewportName)
 			.ScreenPosition(FVector2D(WindowPosition.X, WindowPosition.Y))
