Made the button hide/show based on the Project Setting without restarting the editor.

diff --git a/Engine/Plugins/Runtime/Steam/SteamVR/Source/SteamVREditor/Private/SteamVREditor.cpp b/Engine/Plugins/Runtime/Steam/SteamVR/Source/SteamVREditor/Private/SteamVREditor.cpp
--- a/Engine/Plugins/Runtime/Steam/SteamVR/Source/SteamVREditor/Private/SteamVREditor.cpp
+++ b/Engine/Plugins/Runtime/Steam/SteamVR/Source/SteamVREditor/Private/SteamVREditor.cpp
@@ -44,6 +44,7 @@ POSSIBILITY OF SUCH DAMAGE.
 #include "DetailWidgetRow.h"
 #include "Runtime/Core/Public/CoreGlobals.h"
 #include "Runtime/Core/Public/Misc/ConfigCacheIni.h"
+#include "ToolMenus.h"
 
 static const FName SteamVREditorTabName("SteamVREditor");
 
@@ -96,16 +97,41 @@ void FSteamVREditorModule::StartupModule()
 		FExecuteAction::CreateRaw(this, &FSteamVREditorModule::AddSampleInputs),
 		FCanExecuteAction());
 
+	FSteamVREditorStyle MenuStyle = FSteamVREditorStyle();
+	MenuStyle.Initialize();
+
+	UToolMenu* ToolBar = UToolMenus::Get()->ExtendMenu("LevelEditor.LevelEditorToolBar");
+	FToolMenuSection& Section = ToolBar->FindOrAddSection("Settings");
+	Section.AddEntry(FToolMenuEntry::InitComboButton(
+		"SteamVR Input",
+		FUIAction(
+			FExecuteAction(),
+			FCanExecuteAction(),
+			FGetActionCheckState(),
+			FIsActionButtonVisible::CreateRaw(this, &FSteamVREditorModule::ShowSteamVRInputToolbarDropdown)
+		),
+		FOnGetContent::CreateRaw(this, &FSteamVREditorModule::FillComboButton, PluginCommands),
+		LOCTEXT("SteamVRInputBtn", "SteamVR Input"),
+		LOCTEXT("SteamVRInputBtnTootlip", "SteamVR Input"),
+		FSlateIcon(FSteamVREditorStyle::GetStyleSetName(), "SteamVREditor.PluginAction")
+	));
+}
+
+bool FSteamVREditorModule::ShowSteamVRInputToolbarDropdown()
+{
 	// Only add the toolbar if SteamVR is the currently active tracking system
 	static FName SystemName(TEXT("SteamVR"));
 	if (GEngine && GEngine->XRSystem.IsValid() && (GEngine->XRSystem->GetSystemName() == SystemName))
 	{
-		TSharedPtr<FExtender> ToolbarExtender = MakeShareable(new FExtender);
-		ToolbarExtender->AddToolBarExtension("Settings", EExtensionHook::After, PluginCommands, FToolBarExtensionDelegate::CreateRaw(this, &FSteamVREditorModule::AddToolbarExtension));
-
-		FLevelEditorModule& LevelEditorModule = FModuleManager::LoadModuleChecked<FLevelEditorModule>("LevelEditor");
-		LevelEditorModule.GetToolBarExtensibilityManager()->AddExtender(ToolbarExtender);
+		// Only show the toolbar button if enabled by the user for this project.  Setting is under Project Settings > PLugins > SteamVR now instead.
+		bool bShowSteamVrInputToolbarButton;
+		if (GConfig->GetBool(TEXT("/Script/SteamVREditor.SteamVREditorSettings"), TEXT("bShowSteamVrInputToolbarButton"), bShowSteamVrInputToolbarButton, GEditorIni) && bShowSteamVrInputToolbarButton)
+		{
+			return true;
+		}
 	}
+	
+	return false;
 }
 
 void FSteamVREditorModule::ShutdownModule()
@@ -312,36 +338,6 @@ bool FSteamVREditorModule::AddUniqueActionMapping(TArray<FInputActionKeyMapping>
 	return false;
 }
 
-void FSteamVREditorModule::AddMenuExtension(FMenuBuilder& Builder)
-{
-	Builder.AddMenuEntry(FSteamVREditorCommands::Get().PluginAction);
-}
-
-void FSteamVREditorModule::AddToolbarExtension(FToolBarBuilder& Builder)
-{
-	// Only show the toolbar button if enabled by the user for this project.  Setting is under Project Settings > PLugins > SteamVR now instead.
-	bool bShowSteamVrInputToolbarButton;
-	if (GConfig->GetBool(TEXT("/Script/SteamVREditor.SteamVREditorSettings"), TEXT("bShowSteamVrInputToolbarButton"), bShowSteamVrInputToolbarButton, GEditorIni) && bShowSteamVrInputToolbarButton)
-	{
-		FSteamVREditorStyle MenuStyle = FSteamVREditorStyle();
-		MenuStyle.Initialize();
-
-		static const FName SteamVRInputSectionName = FName("SteamVR Input");
-
-		Builder.BeginSection(SteamVRInputSectionName);
-
-		Builder.AddComboButton(
-			FUIAction(FExecuteAction::CreateRaw(this, &FSteamVREditorModule::PluginButtonClicked)),
-			FOnGetContent::CreateRaw(this, &FSteamVREditorModule::FillComboButton, PluginCommands),
-			LOCTEXT("SteamVRInputBtn", "SteamVR Input"),
-			LOCTEXT("SteamVRInputBtnTootlip", "SteamVR Input"),
-			FSlateIcon(FSteamVREditorStyle::GetStyleSetName(), "SteamVREditor.PluginAction")
-		);
-
-		Builder.EndSection();
-	}
-}
-
 TSharedRef<SWidget> FSteamVREditorModule::FillComboButton(TSharedPtr<class FUICommandList> Commands)
 {
 	FMenuBuilder MenuBuilder(true, Commands);
diff --git a/Engine/Plugins/Runtime/Steam/SteamVR/Source/SteamVREditor/Public/SteamVREditor.h b/Engine/Plugins/Runtime/Steam/SteamVR/Source/SteamVREditor/Public/SteamVREditor.h
--- a/Engine/Plugins/Runtime/Steam/SteamVR/Source/SteamVREditor/Public/SteamVREditor.h
+++ b/Engine/Plugins/Runtime/Steam/SteamVR/Source/SteamVREditor/Public/SteamVREditor.h
@@ -60,10 +60,9 @@ public:
 
 private:
 
-	void AddToolbarExtension(FToolBarBuilder& Builder);
-	void AddMenuExtension(FMenuBuilder& Builder);
 	bool AddUniqueAxisMapping(TArray<FInputAxisKeyMapping> ExistingAxisKeys, UInputSettings* InputSettings, FName ActionName, FKey ActionKey);
 	bool AddUniqueActionMapping(TArray<FInputActionKeyMapping> ExistingActionKeys, UInputSettings* InputSettings, FName ActionName, FKey ActionKey);
+	bool ShowSteamVRInputToolbarDropdown();
 
 	TSharedPtr<class FUICommandList> PluginCommands;
 	TSharedRef<SWidget> FillComboButton(TSharedPtr<class FUICommandList> Commands);
diff --git a/Engine/Plugins/Runtime/Steam/SteamVR/Source/SteamVREditor/SteamVREditor.Build.cs b/Engine/Plugins/Runtime/Steam/SteamVR/Source/SteamVREditor/SteamVREditor.Build.cs
--- a/Engine/Plugins/Runtime/Steam/SteamVR/Source/SteamVREditor/SteamVREditor.Build.cs
+++ b/Engine/Plugins/Runtime/Steam/SteamVR/Source/SteamVREditor/SteamVREditor.Build.cs
@@ -50,7 +50,8 @@ public class SteamVREditor : ModuleRules
                 "SteamVRInput",
                 "SteamVRInputDevice",
                 "BlueprintGraph",
-                "AnimGraph"
+                "AnimGraph",
+				"ToolMenus"
             }
 			);
 
