Disable V-Sync when VR is in use in all plugins using CustomPresent

diff --git a/Engine/Plugins/Runtime/Steam/SteamVR/Source/SteamVR/Private/SteamVRHMD.cpp b/Engine/Plugins/Runtime/Steam/SteamVR/Source/SteamVR/Private/SteamVRHMD.cpp
--- a/Engine/Plugins/Runtime/Steam/SteamVR/Source/SteamVR/Private/SteamVRHMD.cpp
+++ b/Engine/Plugins/Runtime/Steam/SteamVR/Source/SteamVR/Private/SteamVRHMD.cpp
@@ -1792,10 +1792,6 @@ bool FSteamVRHMD::Startup()
 			PixelDensity = FMath::Clamp(PixelDensityCVar->GetFloat(), PixelDensityMin, PixelDensityMax);
 		}
 
-		// disable vsync
-		static IConsoleVariable* CVSyncVar = IConsoleManager::Get().FindConsoleVariable(TEXT("r.VSync"));
-		CVSyncVar->Set(false);
-
 		// enforce finishcurrentframe
 		static IConsoleVariable* CFCFVar = IConsoleManager::Get().FindConsoleVariable(TEXT("r.finishcurrentframe"));
 		CFCFVar->Set(false);
diff --git a/Engine/Plugins/Runtime/Steam/SteamVR/Source/SteamVR/Private/SteamVRRender.cpp b/Engine/Plugins/Runtime/Steam/SteamVR/Source/SteamVR/Private/SteamVRRender.cpp
--- a/Engine/Plugins/Runtime/Steam/SteamVR/Source/SteamVR/Private/SteamVRRender.cpp
+++ b/Engine/Plugins/Runtime/Steam/SteamVR/Source/SteamVR/Private/SteamVRRender.cpp
@@ -153,6 +153,8 @@ bool FSteamVRHMD::BridgeBaseImpl::Present(int& SyncInterval)
 	SwapChain->IncrementSwapChainIndex_RHIThread();
 	DepthSwapChain->IncrementSwapChainIndex_RHIThread();
 
+	SyncInterval = 0;
+
 	return true;
 }
 
diff --git a/Engine/Plugins/Runtime/WindowsMixedReality/Source/WindowsMixedRealityHMD/Private/WindowsMixedRealityCustomPresent.h b/Engine/Plugins/Runtime/WindowsMixedReality/Source/WindowsMixedRealityHMD/Private/WindowsMixedRealityCustomPresent.h
--- a/Engine/Plugins/Runtime/WindowsMixedReality/Source/WindowsMixedRealityHMD/Private/WindowsMixedRealityCustomPresent.h
+++ b/Engine/Plugins/Runtime/WindowsMixedReality/Source/WindowsMixedRealityHMD/Private/WindowsMixedRealityCustomPresent.h
@@ -56,6 +56,8 @@ namespace WindowsMixedReality
 			{
 				hmd->CopyResources(D3D11Context, ViewportTexture);
 			}
+
+			InOutSyncInterval = 0;
 			
 			return hmd->Present();
 #else
