Optimization to use the setting already in memory rather than reading the config file each time the UI is rendered.

diff --git a/Engine/Plugins/Runtime/Steam/SteamVR/Source/SteamVREditor/Private/SteamVREditor.cpp b/Engine/Plugins/Runtime/Steam/SteamVR/Source/SteamVREditor/Private/SteamVREditor.cpp
--- a/Engine/Plugins/Runtime/Steam/SteamVR/Source/SteamVREditor/Private/SteamVREditor.cpp
+++ b/Engine/Plugins/Runtime/Steam/SteamVR/Source/SteamVREditor/Private/SteamVREditor.cpp
@@ -124,8 +124,8 @@ bool FSteamVREditorModule::ShowSteamVRInputToolbarDropdown()
 	if (GEngine && GEngine->XRSystem.IsValid() && (GEngine->XRSystem->GetSystemName() == SystemName))
 	{
 		// Only show the toolbar button if enabled by the user for this project.  Setting is under Project Settings > PLugins > SteamVR now instead.
-		bool bShowSteamVrInputToolbarButton;
-		if (GConfig->GetBool(TEXT("/Script/SteamVREditor.SteamVREditorSettings"), TEXT("bShowSteamVrInputToolbarButton"), bShowSteamVrInputToolbarButton, GEditorIni) && bShowSteamVrInputToolbarButton)
+		USteamVREditorSettings* SteamVREditorSettings = GetMutableDefault<USteamVREditorSettings>();
+		if (SteamVREditorSettings->bShowSteamVrInputToolbarButton)
 		{
 			return true;
 		}
