Should generate the bindings if the manifest doesn't exist.

diff --git a/Engine/Plugins/Runtime/Steam/SteamVR/Source/SteamVRInputDevice/Private/SteamVRInputDevice.cpp b/Engine/Plugins/Runtime/Steam/SteamVR/Source/SteamVRInputDevice/Private/SteamVRInputDevice.cpp
--- a/Engine/Plugins/Runtime/Steam/SteamVR/Source/SteamVRInputDevice/Private/SteamVRInputDevice.cpp
+++ b/Engine/Plugins/Runtime/Steam/SteamVR/Source/SteamVRInputDevice/Private/SteamVRInputDevice.cpp
@@ -2834,7 +2834,7 @@ void FSteamVRInputDevice::GenerateActionManifest(bool GenerateActions, bool Gene
 		{
 			bool bIsGenerated = true;
 
-			if (DeleteIfExists)
+			if (!FileManager.FileExists(*ManifestPath) || DeleteIfExists)
 			{
 				// Set this binding as one we need to regenerate
 				bIsGenerated = false;
