Fix incorrect base pose caching in SteamVR Input module

diff --git a/Engine/Plugins/Runtime/Steam/SteamVR/Source/SteamVRInputDevice/Private/SteamVRInputDevice.cpp b/Engine/Plugins/Runtime/Steam/SteamVR/Source/SteamVRInputDevice/Private/SteamVRInputDevice.cpp
--- a/Engine/Plugins/Runtime/Steam/SteamVR/Source/SteamVRInputDevice/Private/SteamVRInputDevice.cpp
+++ b/Engine/Plugins/Runtime/Steam/SteamVR/Source/SteamVRInputDevice/Private/SteamVRInputDevice.cpp
@@ -169,19 +169,18 @@ void FSteamVRInputDevice::Tick(float DeltaTime)
 		InitSteamVRSystem();
 		//UE_LOG(LogSteamVRInputDevice, Warning, TEXT("SteamVR System Inactive. Trying to initialize..."));
 	}
-	else if(SteamVRHMDModule && SteamVRHMDModule->GetVRSystem())
+	else if(GEngine->XRSystem.IsValid() && SteamVRHMDModule && SteamVRHMDModule->GetVRSystem())
 	{
 		// Cache the controller transform to ensure ResetOrientationAndPosition gets the correct values (Valid for UE4.18 upwards)
 		// [REDACTED_PROFILE_URL]/steamvr_unreal_plugin/issues/2
-		if (GEngine->XRSystem.IsValid() && SteamVRHMDModule && SteamVRHMDModule->GetVRSystem())
-		{
-			CachedBaseOrientation = GEngine->XRSystem->GetBaseOrientation();
-			CachedBasePosition = GEngine->XRSystem->GetBasePosition();
-		}
+		CachedBaseOrientation = GEngine->XRSystem->GetBaseOrientation();
+		CachedBasePosition = GEngine->XRSystem->GetBasePosition();
+	}
+	else
+	{
+		CachedBaseOrientation = FQuat::Identity;
+		CachedBasePosition = FVector::ZeroVector;
 	}
-
-	CachedBaseOrientation = FQuat::Identity;
-	CachedBasePosition = FVector::ZeroVector;
 }
 
 void FSteamVRInputDevice::FindAxisMappings(const UInputSettings* InputSettings, const FName InAxisName, TArray<FInputAxisKeyMapping>& OutMappings) const
