XR: Added -xrtrackingonly commandline to start OpenVR in a mode where frames aren't submitted to allow for reading tracker and device info
This allows multiple instances of the engine to run on a single machine and all share incoming tracking and input data

diff --git a/Engine/Plugins/Runtime/Steam/SteamVR/Source/SteamVR/Private/SteamVRHMD.cpp b/Engine/Plugins/Runtime/Steam/SteamVR/Source/SteamVR/Private/SteamVRHMD.cpp
--- a/Engine/Plugins/Runtime/Steam/SteamVR/Source/SteamVR/Private/SteamVRHMD.cpp
+++ b/Engine/Plugins/Runtime/Steam/SteamVR/Source/SteamVR/Private/SteamVRHMD.cpp
@@ -223,7 +223,17 @@ public:
 	{
 		vr::EVRInitError VRInitErr = vr::VRInitError_None;
 		// Attempt to initialize the VRSystem device
-		VRSystem = vr::VR_Init(&VRInitErr, vr::VRApplication_Scene);
+		// If the commandline -xrtrackingonly is passed, then start the application in _Other mode instead of _Scene mode
+		// This is used when we only want to get tracking information and don't need to render anything to the XR device
+		if (FParse::Param(FCommandLine::Get(), TEXT("xrtrackingonly")))
+		{
+			VRSystem = vr::VR_Init(&VRInitErr, vr::VRApplication_Other);
+			UE_LOG(LogHMD, Log, TEXT("Starting OpenVR in VRApplication_Other mode for tracking only"));
+		}
+		else
+		{
+			VRSystem = vr::VR_Init(&VRInitErr, vr::VRApplication_Scene);
+		}
 		if ((VRSystem == nullptr) || (VRInitErr != vr::VRInitError_None))
 		{
 			UE_LOG(LogHMD, Log, TEXT("Failed to initialize OpenVR with code %d"), (int32)VRInitErr);
