Fixed usage of TCHAR_TO_ANSI and TCHAR_TO_UTF8 in SteamVR.

diff --git a/Engine/Plugins/Runtime/Steam/SteamVR/Source/SteamVR/Private/SteamVRAssetManager.cpp b/Engine/Plugins/Runtime/Steam/SteamVR/Source/SteamVR/Private/SteamVRAssetManager.cpp
--- a/Engine/Plugins/Runtime/Steam/SteamVR/Source/SteamVR/Private/SteamVRAssetManager.cpp
+++ b/Engine/Plugins/Runtime/Steam/SteamVR/Source/SteamVR/Private/SteamVRAssetManager.cpp
@@ -783,7 +783,8 @@ UPrimitiveComponent* FSteamVRAssetManager::CreateRenderComponent(const int32 Dev
 				LoadHandler.BindRaw(this, &FSteamVRAssetManager::OnModelFullyLoaded, ModelName);
 				NewMeshLoader->SetLoaderFinishedCallback(LoadHandler);
 
-				const char* RawModelName = TCHAR_TO_UTF8(*ModelName);
+				const auto RawModelNameConvert = TStringConversion<FTCHARToUTF8_Convert>(*ModelName);
+				const char* RawModelName = RawModelNameConvert.Get();
 				const uint32 SubMeshCount = VRModelManager->GetComponentCount(RawModelName);
 
 				if (SubMeshCount > 0)
diff --git a/Engine/Plugins/Runtime/Steam/SteamVR/Source/SteamVR/Private/SteamVRStereoLayers.cpp b/Engine/Plugins/Runtime/Steam/SteamVR/Source/SteamVR/Private/SteamVRStereoLayers.cpp
--- a/Engine/Plugins/Runtime/Steam/SteamVR/Source/SteamVR/Private/SteamVRStereoLayers.cpp
+++ b/Engine/Plugins/Runtime/Steam/SteamVR/Source/SteamVR/Private/SteamVRStereoLayers.cpp
@@ -82,8 +82,8 @@ void FSteamVRHMD::UpdateLayer(struct FSteamVRLayer& Layer, uint32 LayerId, bool
 	if (bIsValid && Layer.OverlayHandle == vr::k_ulOverlayHandleInvalid)
 	{
 		FString OverlayName = FString::Printf(TEXT("StereoLayer:%u"), LayerId);
-		const char* OverlayNameAnsiStr = TCHAR_TO_ANSI(*OverlayName);
-		OVR_VERIFY(VROverlay->CreateOverlay(OverlayNameAnsiStr, OverlayNameAnsiStr, &Layer.OverlayHandle));
+		const auto OverlayNameAnsiStr = StringCast<ANSICHAR>(*OverlayName);
+		OVR_VERIFY(VROverlay->CreateOverlay(OverlayNameAnsiStr.Get(), OverlayNameAnsiStr.Get(), &Layer.OverlayHandle));
 		OVR_VERIFY(VROverlay->HideOverlay(Layer.OverlayHandle));
 		Layer.bUpdateTexture = true;
 	}
diff --git a/Engine/Plugins/Runtime/Steam/SteamVR/Source/SteamVRInputDevice/Private/SteamVRInputDevice.cpp b/Engine/Plugins/Runtime/Steam/SteamVR/Source/SteamVRInputDevice/Private/SteamVRInputDevice.cpp
--- a/Engine/Plugins/Runtime/Steam/SteamVR/Source/SteamVRInputDevice/Private/SteamVRInputDevice.cpp
+++ b/Engine/Plugins/Runtime/Steam/SteamVR/Source/SteamVRInputDevice/Private/SteamVRInputDevice.cpp
@@ -4335,14 +4335,12 @@ void FSteamVRInputDevice::RegisterApplication(FString ManifestPath)
 
 			GenerateAppManifest(ManifestPath, GameFileName, AppKey, AppManifestPath);
 
-			char* SteamVRAppKey = TCHAR_TO_UTF8(*AppKey);
-
 			// Load application manifest
 			EVRApplicationError AppError = VRApplications()->AddApplicationManifest(TCHAR_TO_UTF8(*IFileManager::Get().ConvertToAbsolutePathForExternalAppForRead(*AppManifestPath)), true);
 			UE_LOG(LogSteamVRInputDevice, Display, TEXT("[STEAMVR INPUT] Registering Application Manifest %s : %s"), *AppManifestPath, *FString(UTF8_TO_TCHAR(VRApplications()->GetApplicationsErrorNameFromEnum(AppError))));
 
 			// Set AppKey for this Editor Session
-			AppError = VRApplications()->IdentifyApplication(AppProcessId, SteamVRAppKey);
+			AppError = VRApplications()->IdentifyApplication(AppProcessId, TCHAR_TO_UTF8(*AppKey));
 			UE_LOG(LogSteamVRInputDevice, Display, TEXT("[STEAMVR INPUT] Editor Application [%d][%s] identified to SteamVR: %s"), AppProcessId, *AppKey, *FString(UTF8_TO_TCHAR(VRApplications()->GetApplicationsErrorNameFromEnum(AppError))));
 		}
 #endif
