Fix right-eye flickering to black when using the SteamVR plugin.

The compositor caches the last result of GetProjectionRaw(), but instead of caching internal values it caches the contents of the output values. By setting the output values to the same memory locations the field-of-view would effectively become zero. Which results in a black screen whenever the compositor uses these field-of-view values to reproject a late frame.

diff --git a/Engine/Plugins/Runtime/Steam/SteamVR/Source/SteamVR/Private/SteamVRHMD.cpp b/Engine/Plugins/Runtime/Steam/SteamVR/Source/SteamVR/Private/SteamVRHMD.cpp
--- a/Engine/Plugins/Runtime/Steam/SteamVR/Source/SteamVR/Private/SteamVRHMD.cpp
+++ b/Engine/Plugins/Runtime/Steam/SteamVR/Source/SteamVR/Private/SteamVRHMD.cpp
@@ -1486,8 +1486,8 @@ FMatrix FSteamVRHMD::GetStereoProjectionMatrix(const int32 StereoViewIndex) cons
 
 	if (StereoViewIndex == EStereoscopicEye::eSSE_MONOSCOPIC)
 	{
-		float Dummy;
-		VRSystem->GetProjectionRaw(vr::Eye_Right, &Right, &Dummy, &Dummy, &Dummy);
+		float DummyLeft, DummyTop, DummyBottom;
+		VRSystem->GetProjectionRaw(vr::Eye_Right, &Right, &DummyLeft, &DummyTop, &DummyBottom);
 		Right *= -1.0f;
 	}
 
