Removed SteamVRInputDevice private pch.
Timing when compiling module with PCH(including compiling PCH): 27.05 secs
Timing using shared PCH: 10.41 secs

diff --git a/Engine/Plugins/Runtime/Steam/SteamVR/Source/SteamVRInputDevice/SteamVRInputDevice.Build.cs b/Engine/Plugins/Runtime/Steam/SteamVR/Source/SteamVRInputDevice/SteamVRInputDevice.Build.cs
--- a/Engine/Plugins/Runtime/Steam/SteamVR/Source/SteamVRInputDevice/SteamVRInputDevice.Build.cs
+++ b/Engine/Plugins/Runtime/Steam/SteamVR/Source/SteamVRInputDevice/SteamVRInputDevice.Build.cs
@@ -35,9 +35,6 @@ public class SteamVRInputDevice : ModuleRules
 {
 	public SteamVRInputDevice(ReadOnlyTargetRules Target) : base(Target)
     {
-        PCHUsage = PCHUsageMode.UseExplicitOrSharedPCHs;
-        PrivatePCHHeaderFile = "Public/ISteamVRInputDeviceModule.h";
-
         PublicIncludePaths.Add(Path.Combine(ModuleDirectory, "Public"));
         PrivateIncludePaths.Add(Path.Combine(ModuleDirectory, "Private"));
 
