Removed SteamVREditor private pch.
Timing when compiling module with PCH(including compiling PCH): 24.99 secs
Timing using shared PCH: 5.85 secs

diff --git a/Engine/Plugins/Runtime/Steam/SteamVR/Source/SteamVREditor/SteamVREditor.Build.cs b/Engine/Plugins/Runtime/Steam/SteamVR/Source/SteamVREditor/SteamVREditor.Build.cs
--- a/Engine/Plugins/Runtime/Steam/SteamVR/Source/SteamVREditor/SteamVREditor.Build.cs
+++ b/Engine/Plugins/Runtime/Steam/SteamVR/Source/SteamVREditor/SteamVREditor.Build.cs
@@ -35,8 +35,6 @@ public class SteamVREditor : ModuleRules
 {
 	public SteamVREditor(ReadOnlyTargetRules Target) : base(Target)
     {
-        PrivatePCHHeaderFile = "Public/SteamVREditor.h";
-
         PublicIncludePaths.Add(Path.Combine(ModuleDirectory, "Public"));
         PrivateIncludePaths.Add(Path.Combine(ModuleDirectory, "Private"));
 
