Update SteamVR input to use the new OnController functions isntead of the old soon to be deprecated versoin

#rnx

diff --git a/Engine/Plugins/Runtime/Steam/SteamVR/Source/SteamVRInputDevice/Private/SteamVRInputDevice.cpp b/Engine/Plugins/Runtime/Steam/SteamVR/Source/SteamVRInputDevice/Private/SteamVRInputDevice.cpp
--- a/Engine/Plugins/Runtime/Steam/SteamVR/Source/SteamVRInputDevice/Private/SteamVRInputDevice.cpp
+++ b/Engine/Plugins/Runtime/Steam/SteamVR/Source/SteamVRInputDevice/Private/SteamVRInputDevice.cpp
@@ -45,6 +45,7 @@ POSSIBILITY OF SUCH DAMAGE.
 #include "SteamVRSkeletonDefinition.h"
 #include "Misc/EngineVersion.h"
 #include "Misc/MessageDialog.h"
+#include "GenericPlatform/GenericPlatformInputDeviceMapper.h"
 
 #if PLATFORM_WINDOWS
 #include "Windows/WindowsHWrapper.h"
@@ -341,14 +342,17 @@ bool FSteamVRInputDevice::GetSkeletalData(bool bLeftHand, bool bMirror, EVRSkele
 
 void FSteamVRInputDevice::SendAnalogMessage(const ETrackedControllerRole TrackedControllerRole, const FGamepadKeyNames::Type AxisButton, float AnalogValue)
 {
+	FInputDeviceId DeviceId = IPlatformInputDeviceMapper::Get().GetDefaultInputDevice();
+	FPlatformUserId UserId = IPlatformInputDeviceMapper::Get().GetPrimaryPlatformUser();
+	
 	if (TrackedControllerRole == ETrackedControllerRole::TrackedControllerRole_LeftHand && bCurlsAndSplaysEnabled_L)
 	{
-		MessageHandler->OnControllerAnalog(AxisButton, 0, AnalogValue);
+		MessageHandler->OnControllerAnalog(AxisButton, UserId, DeviceId, AnalogValue);
 		//UE_LOG(LogSteamVRInputDevice, Warning, TEXT("Left Index value: %f for axis %s"), ControllerState.IndexGripAnalog, *AxisButton.ToString());
 	}
 	else if (TrackedControllerRole == ETrackedControllerRole::TrackedControllerRole_RightHand && bCurlsAndSplaysEnabled_R)
 	{
-		MessageHandler->OnControllerAnalog(AxisButton, 0, AnalogValue);
+		MessageHandler->OnControllerAnalog(AxisButton, UserId, DeviceId, AnalogValue);
 		//UE_LOG(LogSteamVRInputDevice, Warning, TEXT("Left Index value: %f for axis %s"), ControllerState.IndexGripAnalog, *AxisButton.ToString());
 	}
 }
@@ -4522,6 +4526,9 @@ bool FSteamVRInputDevice::SetSkeletalHandle(char* ActionPath, VRActionHandle_t&
 
 void FSteamVRInputDevice::ProcessActionEvents(FSteamVRInputActionSet SteamVRInputActionSet)
 {
+	FInputDeviceId DeviceId = IPlatformInputDeviceMapper::Get().GetDefaultInputDevice();
+	FPlatformUserId UserId = IPlatformInputDeviceMapper::Get().GetPrimaryPlatformUser();
+	
 	for (auto& Action : ActionEvents)
 	{
 		if (Action.Handle == k_ulInvalidActionHandle)
@@ -4553,7 +4560,7 @@ void FSteamVRInputDevice::ProcessActionEvents(FSteamVRInputActionSet SteamVRInpu
 					{
 						if (!Action.bState)
 						{
-							MessageHandler->OnControllerButtonPressed(Action.KeyX, 0, false);
+							MessageHandler->OnControllerButtonPressed(Action.KeyX, UserId, DeviceId, false);
 							Action.bState = DigitalData.bState;
 							Action.LastUpdated = DigitalData.fUpdateTime;
 							Action.bIsRepeat = false;
@@ -4565,7 +4572,7 @@ void FSteamVRInputDevice::ProcessActionEvents(FSteamVRInputActionSet SteamVRInpu
 
 							if (Action.LastUpdated - DigitalData.fUpdateTime >= EffectiveDelay)
 							{
-								MessageHandler->OnControllerButtonPressed(Action.KeyX, 0, true);
+								MessageHandler->OnControllerButtonPressed(Action.KeyX, UserId, DeviceId, true);
 								Action.LastUpdated = DigitalData.fUpdateTime;
 								Action.bIsRepeat = true;
 								//UE_LOG(LogTemp, Warning, TEXT("[REPEAT] Handle %s KeyX %s Value %i Changed %i UpdateTime %f"), *Action.Path, *Action.KeyX.ToString(), DigitalData.bState, DigitalData.bChanged, DigitalData.fUpdateTime);
@@ -4576,7 +4583,7 @@ void FSteamVRInputDevice::ProcessActionEvents(FSteamVRInputActionSet SteamVRInpu
 					{
 						if (Action.bState)
 						{
-							MessageHandler->OnControllerButtonReleased(Action.KeyX, 0, false);
+							MessageHandler->OnControllerButtonReleased(Action.KeyX, UserId, DeviceId, false);
 							//UE_LOG(LogTemp, Display, TEXT("Handle %s KeyX %s Value %i Changed %i UpdateTime %f"), *Action.Path, *Action.KeyX.ToString(), DigitalData.bState, DigitalData.bChanged, DigitalData.fUpdateTime);
 						}
 
@@ -4610,7 +4617,7 @@ void FSteamVRInputDevice::ProcessActionEvents(FSteamVRInputActionSet SteamVRInpu
 					//UE_LOG(LogTemp, Warning, TEXT("Handle %s KeyX %s X-Value [%f]"), *Action.Path, *Action.KeyX.ToString(), AnalogData.x);
 
 					Action.Value.X = AnalogData.x; // ActionCount;
-					MessageHandler->OnControllerAnalog(Action.KeyX, 0, Action.Value.X);
+					MessageHandler->OnControllerAnalog(Action.KeyX, UserId, DeviceId, Action.Value.X);
 				}
 
 				if (Action.KeyY != NAME_None && (FMath::Abs(AnalogData.deltaY) > KINDA_SMALL_NUMBER))
@@ -4627,7 +4634,7 @@ void FSteamVRInputDevice::ProcessActionEvents(FSteamVRInputActionSet SteamVRInpu
 					{
 						Action.Value.Y = AnalogData.y;
 					}
-					MessageHandler->OnControllerAnalog(Action.KeyY, 0, Action.Value.Y);
+					MessageHandler->OnControllerAnalog(Action.KeyY, UserId, DeviceId, Action.Value.Y);
 				}
 			}
 		}
