SteamVR deprecation: SteamVRInput module

diff --git a/Engine/Plugins/Runtime/Steam/SteamVR/Source/SteamVRInput/Private/SteamVRInput.cpp b/Engine/Plugins/Runtime/Steam/SteamVR/Source/SteamVRInput/Private/SteamVRInput.cpp
--- a/Engine/Plugins/Runtime/Steam/SteamVR/Source/SteamVRInput/Private/SteamVRInput.cpp
+++ b/Engine/Plugins/Runtime/Steam/SteamVR/Source/SteamVRInput/Private/SteamVRInput.cpp
@@ -39,6 +39,8 @@ POSSIBILITY OF SUCH DAMAGE.
 
 DEFINE_LOG_CATEGORY_STATIC(LogSteamInput, Log, All);
 
+PRAGMA_DISABLE_DEPRECATION_WARNINGS
+
 void FSteamVRInputModule::StartupModule()
 {
 }
@@ -47,6 +49,9 @@ void FSteamVRInputModule::ShutdownModule()
 {
 }
 
+
 #undef LOCTEXT_NAMESPACE
 	
 IMPLEMENT_MODULE(FSteamVRInputModule, SteamVRInput)
+
+PRAGMA_ENABLE_DEPRECATION_WARNINGS
diff --git a/Engine/Plugins/Runtime/Steam/SteamVR/Source/SteamVRInput/Public/SteamVRInput.h b/Engine/Plugins/Runtime/Steam/SteamVR/Source/SteamVRInput/Public/SteamVRInput.h
--- a/Engine/Plugins/Runtime/Steam/SteamVR/Source/SteamVRInput/Public/SteamVRInput.h
+++ b/Engine/Plugins/Runtime/Steam/SteamVR/Source/SteamVRInput/Public/SteamVRInput.h
@@ -32,6 +32,8 @@ POSSIBILITY OF SUCH DAMAGE.
 
 #include "Modules/ModuleManager.h"
 
+class UE_DEPRECATED(5.1, "SteamVR plugin is deprecated; please use the OpenXR plugin.") FSteamVRInputModule;
+
 class FSteamVRInputModule : public IModuleInterface
 {
 public:
