SteamVR deprecation: SteamVREditor module

diff --git a/Engine/Plugins/Runtime/Steam/SteamVR/Source/SteamVREditor/Private/AnimGraphNode_SteamVRInputAnimPose.cpp b/Engine/Plugins/Runtime/Steam/SteamVR/Source/SteamVREditor/Private/AnimGraphNode_SteamVRInputAnimPose.cpp
--- a/Engine/Plugins/Runtime/Steam/SteamVR/Source/SteamVREditor/Private/AnimGraphNode_SteamVRInputAnimPose.cpp
+++ b/Engine/Plugins/Runtime/Steam/SteamVR/Source/SteamVREditor/Private/AnimGraphNode_SteamVRInputAnimPose.cpp
@@ -33,33 +33,37 @@ POSSIBILITY OF SUCH DAMAGE.
 
 #define LOCTEXT_NAMESPACE "SteamVRInputAnimNode"
 
-UAnimGraphNode_SteamVRInputAnimPose::UAnimGraphNode_SteamVRInputAnimPose(const FObjectInitializer& ObjectInitializer)
+PRAGMA_DISABLE_DEPRECATION_WARNINGS
+
+UDEPRECATED_UAnimGraphNode_SteamVRInputAnimPose::UDEPRECATED_UAnimGraphNode_SteamVRInputAnimPose(const FObjectInitializer& ObjectInitializer)
 	: Super(ObjectInitializer)
 {
 }
 
 // Node Color
-FLinearColor UAnimGraphNode_SteamVRInputAnimPose::GetNodeTitleColor() const 
+FLinearColor UDEPRECATED_UAnimGraphNode_SteamVRInputAnimPose::GetNodeTitleColor() const 
 { 
 	return FLinearColor(0.f, 0.f, 255.f, 1.f);
 }
 
 // Node Category
-FText UAnimGraphNode_SteamVRInputAnimPose::GetMenuCategory() const
+FText UDEPRECATED_UAnimGraphNode_SteamVRInputAnimPose::GetMenuCategory() const
 {
 	return LOCTEXT("NodeCategory", "SteamVR Input");
 }
 
 // Node Title
-FText UAnimGraphNode_SteamVRInputAnimPose::GetNodeTitle(ENodeTitleType::Type TitleType) const
+FText UDEPRECATED_UAnimGraphNode_SteamVRInputAnimPose::GetNodeTitle(ENodeTitleType::Type TitleType) const
 {
 	return LOCTEXT("NodeTitle", "SteamVR Skeletal Anim Pose");
 }
 
 // Node Tooltip
-FText UAnimGraphNode_SteamVRInputAnimPose::GetTooltipText() const
+FText UDEPRECATED_UAnimGraphNode_SteamVRInputAnimPose::GetTooltipText() const
 {
 	return LOCTEXT("NodeTooltip", "Retrieves the current pose from the SteamVR Skeletal Input API");
 }
 
+PRAGMA_ENABLE_DEPRECATION_WARNINGS
+
 #undef LOCTEXT_NAMESPACE
\ No newline at end of file
diff --git a/Engine/Plugins/Runtime/Steam/SteamVR/Source/SteamVREditor/Private/AnimGraphNode_SteamVRSetWristTransform.cpp b/Engine/Plugins/Runtime/Steam/SteamVR/Source/SteamVREditor/Private/AnimGraphNode_SteamVRSetWristTransform.cpp
--- a/Engine/Plugins/Runtime/Steam/SteamVR/Source/SteamVREditor/Private/AnimGraphNode_SteamVRSetWristTransform.cpp
+++ b/Engine/Plugins/Runtime/Steam/SteamVR/Source/SteamVREditor/Private/AnimGraphNode_SteamVRSetWristTransform.cpp
@@ -33,33 +33,37 @@ POSSIBILITY OF SUCH DAMAGE.
 
 #define LOCTEXT_NAMESPACE "SteamVRSetWristTransformAnimNode"
 
-UAnimGraphNode_SteamVRSetWristTransform::UAnimGraphNode_SteamVRSetWristTransform(const FObjectInitializer& ObjectInitializer)
+PRAGMA_DISABLE_DEPRECATION_WARNINGS
+
+UDEPRECATED_UAnimGraphNode_SteamVRSetWristTransform::UDEPRECATED_UAnimGraphNode_SteamVRSetWristTransform(const FObjectInitializer& ObjectInitializer)
 	: Super(ObjectInitializer)
 {
 }
 
 // Node Color
-FLinearColor UAnimGraphNode_SteamVRSetWristTransform::GetNodeTitleColor() const 
+FLinearColor UDEPRECATED_UAnimGraphNode_SteamVRSetWristTransform::GetNodeTitleColor() const 
 { 
 	return FLinearColor(0.f, 0.f, 255.f, 1.f);
 }
 
 // Node Category
-FText UAnimGraphNode_SteamVRSetWristTransform::GetMenuCategory() const
+FText UDEPRECATED_UAnimGraphNode_SteamVRSetWristTransform::GetMenuCategory() const
 {
 	return LOCTEXT("NodeCategory", "SteamVR Input");
 }
 
 // Node Title
-FText UAnimGraphNode_SteamVRSetWristTransform::GetNodeTitle(ENodeTitleType::Type TitleType) const
+FText UDEPRECATED_UAnimGraphNode_SteamVRSetWristTransform::GetNodeTitle(ENodeTitleType::Type TitleType) const
 {
 	return LOCTEXT("NodeTitle", "SteamVR Set Wrist Transform");
 }
 
 // Node Tooltip
-FText UAnimGraphNode_SteamVRSetWristTransform::GetTooltipText() const
+FText UDEPRECATED_UAnimGraphNode_SteamVRSetWristTransform::GetTooltipText() const
 {
 	return LOCTEXT("NodeTooltip", "Sets the wrist transform of target pose from the reference pose");
 }
 
+PRAGMA_ENABLE_DEPRECATION_WARNINGS
+
 #undef LOCTEXT_NAMESPACE
\ No newline at end of file
diff --git a/Engine/Plugins/Runtime/Steam/SteamVR/Source/SteamVREditor/Private/SteamVREditor.cpp b/Engine/Plugins/Runtime/Steam/SteamVR/Source/SteamVREditor/Private/SteamVREditor.cpp
--- a/Engine/Plugins/Runtime/Steam/SteamVR/Source/SteamVREditor/Private/SteamVREditor.cpp
+++ b/Engine/Plugins/Runtime/Steam/SteamVR/Source/SteamVREditor/Private/SteamVREditor.cpp
@@ -50,6 +50,8 @@ static const FName SteamVREditorTabName("SteamVREditor");
 
 #define LOCTEXT_NAMESPACE "FSteamVREditorModule"
 
+PRAGMA_DISABLE_DEPRECATION_WARNINGS
+
 void FSteamVREditorModule::StartupModule()
 {
 	RegisterSettings();
@@ -124,7 +126,7 @@ bool FSteamVREditorModule::ShowSteamVRInputToolbarDropdown()
 	if (GEngine && GEngine->XRSystem.IsValid() && (GEngine->XRSystem->GetSystemName() == SystemName))
 	{
 		// Only show the toolbar button if enabled by the user for this project.  Setting is under Project Settings > PLugins > SteamVR now instead.
-		USteamVREditorSettings* SteamVREditorSettings = GetMutableDefault<USteamVREditorSettings>();
+		UDEPRECATED_USteamVREditorSettings* SteamVREditorSettings = GetMutableDefault<UDEPRECATED_USteamVREditorSettings>();
 		if (SteamVREditorSettings->bShowSteamVrInputToolbarButton)
 		{
 			return true;
@@ -358,11 +360,11 @@ void FSteamVREditorModule::RegisterSettings()
 		SettingsModule->RegisterSettings("Project", "Plugins", "SteamVR",
 			LOCTEXT("SteamVREditorSettingsName", "SteamVR"),
 			LOCTEXT("SteamVREditorSettingsDescription", "Configure the SteamVR plugin"),
-			GetMutableDefault<USteamVREditorSettings>()
+			GetMutableDefault<UDEPRECATED_USteamVREditorSettings>()
 		);
 
 		FPropertyEditorModule& PropertyModule = FModuleManager::GetModuleChecked<FPropertyEditorModule>("PropertyEditor");
-		PropertyModule.RegisterCustomClassLayout(USteamVREditorSettings::StaticClass()->GetFName(), FOnGetDetailCustomizationInstance::CreateStatic(&FSteamVRSettingsDetailsCustomization::MakeInstance));
+		PropertyModule.RegisterCustomClassLayout(UDEPRECATED_USteamVREditorSettings::StaticClass()->GetFName(), FOnGetDetailCustomizationInstance::CreateStatic(&FSteamVRSettingsDetailsCustomization::MakeInstance));
 	}
 }
 
@@ -468,4 +470,6 @@ FReply FSteamVRSettingsDetailsCustomization::LaunchBindingsURL()
 	return FReply::Handled();
 }
 
+PRAGMA_ENABLE_DEPRECATION_WARNINGS
+
 #undef LOCTEXT_NAMESPACE
diff --git a/Engine/Plugins/Runtime/Steam/SteamVR/Source/SteamVREditor/Public/AnimGraphNode_SteamVRInputAnimPose.h b/Engine/Plugins/Runtime/Steam/SteamVR/Source/SteamVREditor/Public/AnimGraphNode_SteamVRInputAnimPose.h
--- a/Engine/Plugins/Runtime/Steam/SteamVR/Source/SteamVREditor/Public/AnimGraphNode_SteamVRInputAnimPose.h
+++ b/Engine/Plugins/Runtime/Steam/SteamVR/Source/SteamVREditor/Public/AnimGraphNode_SteamVRInputAnimPose.h
@@ -37,12 +37,13 @@ POSSIBILITY OF SUCH DAMAGE.
 
 using namespace vr;
 
-UCLASS()
-class UAnimGraphNode_SteamVRInputAnimPose : public UAnimGraphNode_Base
+UCLASS(deprecated, meta = (DeprecationMessage = "SteamVR plugin is deprecated; please use the OpenXR plugin."))
+class UDEPRECATED_UAnimGraphNode_SteamVRInputAnimPose : public UAnimGraphNode_Base
 {
 	GENERATED_UCLASS_BODY()
 
-	UPROPERTY(EditAnywhere, Category=Settings)
+	UE_DEPRECATED(5.1, "SteamVR plugin is deprecated; please use the OpenXR plugin.")
+	UPROPERTY(EditAnywhere, Category=Settings, meta = (DeprecatedProperty))
 	FAnimNode_SteamVRInputAnimPose Node;
 
 public:
diff --git a/Engine/Plugins/Runtime/Steam/SteamVR/Source/SteamVREditor/Public/AnimGraphNode_SteamVRSetWristTransform.h b/Engine/Plugins/Runtime/Steam/SteamVR/Source/SteamVREditor/Public/AnimGraphNode_SteamVRSetWristTransform.h
--- a/Engine/Plugins/Runtime/Steam/SteamVR/Source/SteamVREditor/Public/AnimGraphNode_SteamVRSetWristTransform.h
+++ b/Engine/Plugins/Runtime/Steam/SteamVR/Source/SteamVREditor/Public/AnimGraphNode_SteamVRSetWristTransform.h
@@ -37,12 +37,13 @@ POSSIBILITY OF SUCH DAMAGE.
 
 using namespace vr;
 
-UCLASS()
-class UAnimGraphNode_SteamVRSetWristTransform : public UAnimGraphNode_Base
+UCLASS(deprecated, meta = (DeprecationMessage = "SteamVR plugin is deprecated; please use the OpenXR plugin."))
+class UDEPRECATED_UAnimGraphNode_SteamVRSetWristTransform : public UAnimGraphNode_Base
 {
 	GENERATED_UCLASS_BODY()
 
-	UPROPERTY(EditAnywhere, Category=Settings)
+	UE_DEPRECATED(5.1, "SteamVR plugin is deprecated; please use the OpenXR plugin.")
+	UPROPERTY(EditAnywhere, Category=Settings, meta = (DeprecatedProperty))
 	FAnimNode_SteamVRSetWristTransform Node;
 
 public:
diff --git a/Engine/Plugins/Runtime/Steam/SteamVR/Source/SteamVREditor/Public/SteamVREditor.h b/Engine/Plugins/Runtime/Steam/SteamVR/Source/SteamVREditor/Public/SteamVREditor.h
--- a/Engine/Plugins/Runtime/Steam/SteamVR/Source/SteamVREditor/Public/SteamVREditor.h
+++ b/Engine/Plugins/Runtime/Steam/SteamVR/Source/SteamVREditor/Public/SteamVREditor.h
@@ -38,6 +38,7 @@ class FToolBarBuilder;
 class FMenuBuilder;
 class SWidget;
 
+class UE_DEPRECATED(5.1, "SteamVR plugin is deprecated; please use the OpenXR plugin.") FSteamVREditorModule;
 class FSteamVREditorModule : public IModuleInterface
 {
 public:
@@ -69,6 +70,7 @@ private:
 
 };
 
+class UE_DEPRECATED(5.1, "SteamVR plugin is deprecated; please use the OpenXR plugin.") FSteamVRSettingsDetailsCustomization;
 class FSteamVRSettingsDetailsCustomization : public IDetailCustomization
 {
 public:
diff --git a/Engine/Plugins/Runtime/Steam/SteamVR/Source/SteamVREditor/Public/SteamVREditorSettings.h b/Engine/Plugins/Runtime/Steam/SteamVR/Source/SteamVREditor/Public/SteamVREditorSettings.h
--- a/Engine/Plugins/Runtime/Steam/SteamVR/Source/SteamVREditor/Public/SteamVREditorSettings.h
+++ b/Engine/Plugins/Runtime/Steam/SteamVR/Source/SteamVREditor/Public/SteamVREditorSettings.h
@@ -9,12 +9,13 @@
 /**
  * 
  */
-UCLASS(config = Editor, defaultconfig)
-class STEAMVREDITOR_API USteamVREditorSettings : public UObject
+UCLASS(config = Editor, defaultconfig, deprecated, meta = (DeprecationMessage = "SteamVR plugin is deprecated; please use the OpenXR plugin."))
+class STEAMVREDITOR_API UDEPRECATED_USteamVREditorSettings : public UObject
 {
 	GENERATED_BODY()
 public:
 	/** Whether or not to show the SteamVR Input settings toolbar button */
-	UPROPERTY(config, EditAnywhere, Category = "SteamVR Editor Settings")
+	UE_DEPRECATED(5.1, "SteamVR plugin is deprecated; please use the OpenXR plugin.")
+	UPROPERTY(config, EditAnywhere, Category = "SteamVR Editor Settings", meta = (DeprecatedProperty))
 	bool bShowSteamVrInputToolbarButton;
 };
