SteamVR deprecation: SteanVRInputDevice module

diff --git a/Engine/Plugins/Runtime/Steam/SteamVR/Source/SteamVREditor/Private/SteamVREditor.cpp b/Engine/Plugins/Runtime/Steam/SteamVR/Source/SteamVREditor/Private/SteamVREditor.cpp
--- a/Engine/Plugins/Runtime/Steam/SteamVR/Source/SteamVREditor/Private/SteamVREditor.cpp
+++ b/Engine/Plugins/Runtime/Steam/SteamVR/Source/SteamVREditor/Private/SteamVREditor.cpp
@@ -150,22 +150,22 @@ void FSteamVREditorModule::PluginButtonClicked()
 
 void FSteamVREditorModule::JsonRegenerateActionManifest()
 {
-	USteamVRInputDeviceFunctionLibrary::RegenActionManifest();
+	UDEPRECATED_USteamVRInputDeviceFunctionLibrary::RegenActionManifest();
 }
 
 void FSteamVREditorModule::JsonRegenerateControllerBindings()
 {
-	USteamVRInputDeviceFunctionLibrary::RegenControllerBindings();
+	UDEPRECATED_USteamVRInputDeviceFunctionLibrary::RegenControllerBindings();
 }
 
 void FSteamVREditorModule::ReloadActionManifest()
 {
-	USteamVRInputDeviceFunctionLibrary::ReloadActionManifest();
+	UDEPRECATED_USteamVRInputDeviceFunctionLibrary::ReloadActionManifest();
 }
 
 void FSteamVREditorModule::LaunchBindingsURL()
 {
-	USteamVRInputDeviceFunctionLibrary::LaunchBindingsURL();
+	UDEPRECATED_USteamVRInputDeviceFunctionLibrary::LaunchBindingsURL();
 }
 
 void FSteamVREditorModule::AddSampleInputs()
diff --git a/Engine/Plugins/Runtime/Steam/SteamVR/Source/SteamVRInputDevice/Private/AnimNode_SteamVRInputAnimPose.cpp b/Engine/Plugins/Runtime/Steam/SteamVR/Source/SteamVRInputDevice/Private/AnimNode_SteamVRInputAnimPose.cpp
--- a/Engine/Plugins/Runtime/Steam/SteamVR/Source/SteamVRInputDevice/Private/AnimNode_SteamVRInputAnimPose.cpp
+++ b/Engine/Plugins/Runtime/Steam/SteamVR/Source/SteamVRInputDevice/Private/AnimNode_SteamVRInputAnimPose.cpp
@@ -35,6 +35,8 @@ POSSIBILITY OF SUCH DAMAGE.
 #include "SteamVRInputDevice.h"
 #include "UE4HandSkeletonDefinition.h"
 
+PRAGMA_DISABLE_DEPRECATION_WARNINGS
+
 // Lookup table that maps each bone in the UE4 hand skeleton to the corresponding
 // bone in the SteamVR hand skeleton
 static const int32 kUE4BoneToSteamVRBone[] = {
@@ -342,3 +344,5 @@ FTransform FAnimNode_SteamVRInputAnimPose::CalcModelSpaceTransform(const FCompac
 
 	return BoneTransform;
 }
+
+PRAGMA_ENABLE_DEPRECATION_WARNINGS
diff --git a/Engine/Plugins/Runtime/Steam/SteamVR/Source/SteamVRInputDevice/Private/AnimNode_SteamVRSetWristTransform.cpp b/Engine/Plugins/Runtime/Steam/SteamVR/Source/SteamVRInputDevice/Private/AnimNode_SteamVRSetWristTransform.cpp
--- a/Engine/Plugins/Runtime/Steam/SteamVR/Source/SteamVRInputDevice/Private/AnimNode_SteamVRSetWristTransform.cpp
+++ b/Engine/Plugins/Runtime/Steam/SteamVR/Source/SteamVRInputDevice/Private/AnimNode_SteamVRSetWristTransform.cpp
@@ -35,6 +35,8 @@ POSSIBILITY OF SUCH DAMAGE.
 #include "SteamVRInputDevice.h"
 #include "UE4HandSkeletonDefinition.h"
 
+PRAGMA_DISABLE_DEPRECATION_WARNINGS
+
 FAnimNode_SteamVRSetWristTransform::FAnimNode_SteamVRSetWristTransform()
 {
 }
@@ -83,3 +85,5 @@ void FAnimNode_SteamVRSetWristTransform::Evaluate_AnyThread(FPoseContext& Output
 		Output.Pose[RootBoneIndex] = ReferenceContext.Pose[RootBoneIndex];
 	}
 }
+
+PRAGMA_ENABLE_DEPRECATION_WARNINGS
diff --git a/Engine/Plugins/Runtime/Steam/SteamVR/Source/SteamVRInputDevice/Private/SteamVRInputDevice.cpp b/Engine/Plugins/Runtime/Steam/SteamVR/Source/SteamVRInputDevice/Private/SteamVRInputDevice.cpp
--- a/Engine/Plugins/Runtime/Steam/SteamVR/Source/SteamVRInputDevice/Private/SteamVRInputDevice.cpp
+++ b/Engine/Plugins/Runtime/Steam/SteamVR/Source/SteamVRInputDevice/Private/SteamVRInputDevice.cpp
@@ -61,6 +61,8 @@ POSSIBILITY OF SUCH DAMAGE.
 #define LOCTEXT_NAMESPACE "SteamVRInputDevice"
 DEFINE_LOG_CATEGORY_STATIC(LogSteamVRInputDevice, Log, All);
 
+PRAGMA_DISABLE_DEPRECATION_WARNINGS
+
 // List of bones that are effectively in model space because
 // they are children of the root
 static const int32 kModelSpaceBones[] = {
@@ -4781,4 +4783,6 @@ bool FSteamVRInputDevice::ProcessVector2D(FSteamVRAxisKeyMapping AxisKeyMapping)
 	return false;
 }
 
+PRAGMA_ENABLE_DEPRECATION_WARNINGS
+
 #undef LOCTEXT_NAMESPACE //"SteamVRInputDevice"
diff --git a/Engine/Plugins/Runtime/Steam/SteamVR/Source/SteamVRInputDevice/Private/SteamVRInputDeviceFunctionLibrary.cpp b/Engine/Plugins/Runtime/Steam/SteamVR/Source/SteamVRInputDevice/Private/SteamVRInputDeviceFunctionLibrary.cpp
--- a/Engine/Plugins/Runtime/Steam/SteamVR/Source/SteamVRInputDevice/Private/SteamVRInputDeviceFunctionLibrary.cpp
+++ b/Engine/Plugins/Runtime/Steam/SteamVR/Source/SteamVRInputDevice/Private/SteamVRInputDeviceFunctionLibrary.cpp
@@ -38,7 +38,9 @@ POSSIBILITY OF SUCH DAMAGE.
 using namespace vr;
 //#endif // STEAMVRCONTROLLER_SUPPORTED_PLATFORMS
 
-void USteamVRInputDeviceFunctionLibrary::PlaySteamVR_HapticFeedback(ESteamVRHand Hand, float StartSecondsFromNow, float DurationSeconds, float Frequency, float Amplitude)
+PRAGMA_DISABLE_DEPRECATION_WARNINGS
+
+void UDEPRECATED_USteamVRInputDeviceFunctionLibrary::PlaySteamVR_HapticFeedback(ESteamVRHand Hand, float StartSecondsFromNow, float DurationSeconds, float Frequency, float Amplitude)
 {
 	FSteamVRInputDevice* SteamVRInputDevice = GetSteamVRInputDevice();
 	if (VRSystem() && VRInput())
@@ -68,7 +70,7 @@ void USteamVRInputDeviceFunctionLibrary::PlaySteamVR_HapticFeedback(ESteamVRHand
 	}
 }
 
-void USteamVRInputDeviceFunctionLibrary::RegenActionManifest()
+void UDEPRECATED_USteamVRInputDeviceFunctionLibrary::RegenActionManifest()
 {
 #if WITH_EDITOR
 	FSteamVRInputDevice* SteamVRInputDevice = GetSteamVRInputDevice();
@@ -79,7 +81,7 @@ void USteamVRInputDeviceFunctionLibrary::RegenActionManifest()
 #endif
 }
 
-void USteamVRInputDeviceFunctionLibrary::RegenControllerBindings()
+void UDEPRECATED_USteamVRInputDeviceFunctionLibrary::RegenControllerBindings()
 {
 #if WITH_EDITOR
 	FSteamVRInputDevice* SteamVRInputDevice = GetSteamVRInputDevice();
@@ -90,7 +92,7 @@ void USteamVRInputDeviceFunctionLibrary::RegenControllerBindings()
 #endif
 }
 
-void USteamVRInputDeviceFunctionLibrary::ReloadActionManifest()
+void UDEPRECATED_USteamVRInputDeviceFunctionLibrary::ReloadActionManifest()
 {
 #if WITH_EDITOR
 	FSteamVRInputDevice* SteamVRInputDevice = GetSteamVRInputDevice();
@@ -101,7 +103,7 @@ void USteamVRInputDeviceFunctionLibrary::ReloadActionManifest()
 #endif
 }
 
-void USteamVRInputDeviceFunctionLibrary::GetCurlsAndSplaysState(bool& LeftHandState, bool& RightHandState)
+void UDEPRECATED_USteamVRInputDeviceFunctionLibrary::GetCurlsAndSplaysState(bool& LeftHandState, bool& RightHandState)
 {
 	FSteamVRInputDevice* SteamVRInputDevice = GetSteamVRInputDevice();
 	if (SteamVRInputDevice != nullptr)
@@ -111,7 +113,7 @@ void USteamVRInputDeviceFunctionLibrary::GetCurlsAndSplaysState(bool& LeftHandSt
 	}
 }
 
-void USteamVRInputDeviceFunctionLibrary::GetSkeletalState(bool& LeftHandState, bool& RightHandState)
+void UDEPRECATED_USteamVRInputDeviceFunctionLibrary::GetSkeletalState(bool& LeftHandState, bool& RightHandState)
 {
 	FSteamVRInputDevice* SteamVRInputDevice = GetSteamVRInputDevice();
 	if (SteamVRInputDevice != nullptr)
@@ -121,7 +123,7 @@ void USteamVRInputDeviceFunctionLibrary::GetSkeletalState(bool& LeftHandState, b
 	}
 }
 
-void USteamVRInputDeviceFunctionLibrary::GetControllerFidelity(EControllerFidelity& LeftControllerFidelity, EControllerFidelity& RightControllerFidelity)
+void UDEPRECATED_USteamVRInputDeviceFunctionLibrary::GetControllerFidelity(EControllerFidelity& LeftControllerFidelity, EControllerFidelity& RightControllerFidelity)
 {
 	FSteamVRInputDevice* SteamVRInputDevice = GetSteamVRInputDevice();
 	if (SteamVRInputDevice != nullptr)
@@ -167,7 +169,7 @@ void USteamVRInputDeviceFunctionLibrary::GetControllerFidelity(EControllerFideli
 	}
 }
 
-void USteamVRInputDeviceFunctionLibrary::SetCurlsAndSplaysState(bool NewLeftHandState, bool NewRightHandState)
+void UDEPRECATED_USteamVRInputDeviceFunctionLibrary::SetCurlsAndSplaysState(bool NewLeftHandState, bool NewRightHandState)
 {
 	FSteamVRInputDevice* SteamVRInputDevice = GetSteamVRInputDevice();
 	if (SteamVRInputDevice != nullptr)
@@ -177,7 +179,7 @@ void USteamVRInputDeviceFunctionLibrary::SetCurlsAndSplaysState(bool NewLeftHand
 	}
 }
 
-void USteamVRInputDeviceFunctionLibrary::GetPoseSource(bool& bUsingSkeletonPose)
+void UDEPRECATED_USteamVRInputDeviceFunctionLibrary::GetPoseSource(bool& bUsingSkeletonPose)
 {
 	FSteamVRInputDevice* SteamVRInputDevice = GetSteamVRInputDevice();
 	if (SteamVRInputDevice != nullptr)
@@ -186,7 +188,7 @@ void USteamVRInputDeviceFunctionLibrary::GetPoseSource(bool& bUsingSkeletonPose)
 	}
 }
 
-void USteamVRInputDeviceFunctionLibrary::SetPoseSource(bool bUseSkeletonPose)
+void UDEPRECATED_USteamVRInputDeviceFunctionLibrary::SetPoseSource(bool bUseSkeletonPose)
 {
 	FSteamVRInputDevice* SteamVRInputDevice = GetSteamVRInputDevice();
 	if (SteamVRInputDevice != nullptr)
@@ -195,7 +197,7 @@ void USteamVRInputDeviceFunctionLibrary::SetPoseSource(bool bUseSkeletonPose)
 	}
 }
 
-void USteamVRInputDeviceFunctionLibrary::GetSkeletalTransform(FSteamVRSkeletonTransform& LeftHand, FSteamVRSkeletonTransform& RightHand, bool bWithController)
+void UDEPRECATED_USteamVRInputDeviceFunctionLibrary::GetSkeletalTransform(FSteamVRSkeletonTransform& LeftHand, FSteamVRSkeletonTransform& RightHand, bool bWithController)
 {
 	FTransform OutPose[STEAMVR_SKELETON_BONE_COUNT];
 
@@ -289,7 +291,7 @@ void USteamVRInputDeviceFunctionLibrary::GetSkeletalTransform(FSteamVRSkeletonTr
 	}
 }
 
-void USteamVRInputDeviceFunctionLibrary::GetLeftHandPoseData(FVector& Position, FRotator& Orientation, FVector& AngularVelocity, FVector& Velocity)
+void UDEPRECATED_USteamVRInputDeviceFunctionLibrary::GetLeftHandPoseData(FVector& Position, FRotator& Orientation, FVector& AngularVelocity, FVector& Velocity)
 {
 	FSteamVRInputDevice* SteamVRInputDevice = GetSteamVRInputDevice();
 	if (SteamVRInputDevice != nullptr)
@@ -298,7 +300,7 @@ void USteamVRInputDeviceFunctionLibrary::GetLeftHandPoseData(FVector& Position,
 	}
 }
 
-void USteamVRInputDeviceFunctionLibrary::GetRightHandPoseData(FVector& Position, FRotator& Orientation, FVector& AngularVelocity, FVector& Velocity)
+void UDEPRECATED_USteamVRInputDeviceFunctionLibrary::GetRightHandPoseData(FVector& Position, FRotator& Orientation, FVector& AngularVelocity, FVector& Velocity)
 {
 	FSteamVRInputDevice* SteamVRInputDevice = GetSteamVRInputDevice();
 	if (SteamVRInputDevice != nullptr)
@@ -307,7 +309,7 @@ void USteamVRInputDeviceFunctionLibrary::GetRightHandPoseData(FVector& Position,
 	}
 }
 
-void USteamVRInputDeviceFunctionLibrary::GetSteamVR_ActionArray(TArray<FSteamVRAction>& SteamVRActions)
+void UDEPRECATED_USteamVRInputDeviceFunctionLibrary::GetSteamVR_ActionArray(TArray<FSteamVRAction>& SteamVRActions)
 {
 	bool bAlreadyExists;
 	FSteamVRInputDevice* SteamVRInputDevice = GetSteamVRInputDevice();
@@ -336,7 +338,7 @@ void USteamVRInputDeviceFunctionLibrary::GetSteamVR_ActionArray(TArray<FSteamVRA
 	}
 }
 
-void USteamVRInputDeviceFunctionLibrary::FindSteamVR_Action(FName ActionName, bool& bResult, FSteamVRAction& FoundAction, FSteamVRActionSet& FoundActionSet, FName ActionSet)
+void UDEPRECATED_USteamVRInputDeviceFunctionLibrary::FindSteamVR_Action(FName ActionName, bool& bResult, FSteamVRAction& FoundAction, FSteamVRActionSet& FoundActionSet, FName ActionSet)
 {
 	bool bActionFound = false;
 	bool bActionSetFound = false;
@@ -376,7 +378,7 @@ void USteamVRInputDeviceFunctionLibrary::FindSteamVR_Action(FName ActionName, bo
 	bResult = bActionFound && bActionSetFound;
 }
 
-void USteamVRInputDeviceFunctionLibrary::GetSteamVR_ActionSetArray(TArray<FSteamVRActionSet>& SteamVRActionSets)
+void UDEPRECATED_USteamVRInputDeviceFunctionLibrary::GetSteamVR_ActionSetArray(TArray<FSteamVRActionSet>& SteamVRActionSets)
 {
 	// We only have one Action Set supported for this version of the plugin, so just return that
 	FSteamVRInputDevice* SteamVRInputDevice = GetSteamVRInputDevice();
@@ -386,7 +388,7 @@ void USteamVRInputDeviceFunctionLibrary::GetSteamVR_ActionSetArray(TArray<FSteam
 	}
 }
 
-bool USteamVRInputDeviceFunctionLibrary::GetSteamVR_OriginTrackedDeviceInfo(FSteamVRAction SteamVRAction, FSteamVRInputOriginInfo& InputOriginInfo)
+bool UDEPRECATED_USteamVRInputDeviceFunctionLibrary::GetSteamVR_OriginTrackedDeviceInfo(FSteamVRAction SteamVRAction, FSteamVRInputOriginInfo& InputOriginInfo)
 {
 	if (VRSystem() && VRInput())
 	{
@@ -414,7 +416,7 @@ bool USteamVRInputDeviceFunctionLibrary::GetSteamVR_OriginTrackedDeviceInfo(FSte
 	return false;
 }
 
-void USteamVRInputDeviceFunctionLibrary::FindSteamVR_OriginTrackedDeviceInfo(FName ActionName, bool& bResult, FSteamVRInputOriginInfo& InputOriginInfo, FName ActionSet /*= FName("main")*/)
+void UDEPRECATED_USteamVRInputDeviceFunctionLibrary::FindSteamVR_OriginTrackedDeviceInfo(FName ActionName, bool& bResult, FSteamVRInputOriginInfo& InputOriginInfo, FName ActionSet /*= FName("main")*/)
 {
 	bResult = false;
 	FSteamVRAction FoundAction;
@@ -428,7 +430,7 @@ void USteamVRInputDeviceFunctionLibrary::FindSteamVR_OriginTrackedDeviceInfo(FNa
 	}
 }
 
-void USteamVRInputDeviceFunctionLibrary::GetSteamVR_OriginLocalizedName(FSteamVRAction SteamVRAction, TArray<ESteamVRInputStringBits> LocalizedParts, FString& OriginLocalizedName)
+void UDEPRECATED_USteamVRInputDeviceFunctionLibrary::GetSteamVR_OriginLocalizedName(FSteamVRAction SteamVRAction, TArray<ESteamVRInputStringBits> LocalizedParts, FString& OriginLocalizedName)
 {
 	if (VRSystem() && VRInput())
 	{
@@ -475,7 +477,7 @@ void USteamVRInputDeviceFunctionLibrary::GetSteamVR_OriginLocalizedName(FSteamVR
 	}
 }
 
-void USteamVRInputDeviceFunctionLibrary::ShowSteamVR_ActionOrigin(FSteamVRAction SteamVRAction, FSteamVRActionSet SteamVRActionSet)
+void UDEPRECATED_USteamVRInputDeviceFunctionLibrary::ShowSteamVR_ActionOrigin(FSteamVRAction SteamVRAction, FSteamVRActionSet SteamVRActionSet)
 {
 	if (VRSystem() && VRInput())
 	{
@@ -494,7 +496,7 @@ void USteamVRInputDeviceFunctionLibrary::ShowSteamVR_ActionOrigin(FSteamVRAction
 	}
 }
 
-bool USteamVRInputDeviceFunctionLibrary::FindSteamVR_ActionOrigin(FName ActionName, FName ActionSet)
+bool UDEPRECATED_USteamVRInputDeviceFunctionLibrary::FindSteamVR_ActionOrigin(FName ActionName, FName ActionSet)
 {
 	// Find SteamVR Action
 	bool bIsActionFound = false;
@@ -512,7 +514,7 @@ bool USteamVRInputDeviceFunctionLibrary::FindSteamVR_ActionOrigin(FName ActionNa
 	return false;
 }
 
-bool USteamVRInputDeviceFunctionLibrary::GetSteamVR_HandPoseRelativeToNow(FVector& Position, FRotator& Orientation, ESteamVRHand Hand /*= ESteamVRHand::VR_Left*/, float PredictedSecondsFromNow /*= 0.f*/)
+bool UDEPRECATED_USteamVRInputDeviceFunctionLibrary::GetSteamVR_HandPoseRelativeToNow(FVector& Position, FRotator& Orientation, ESteamVRHand Hand /*= ESteamVRHand::VR_Left*/, float PredictedSecondsFromNow /*= 0.f*/)
 {
 	if (VRSystem() && VRInput())
 	{
@@ -564,7 +566,7 @@ bool USteamVRInputDeviceFunctionLibrary::GetSteamVR_HandPoseRelativeToNow(FVecto
 	return false;
 }
 
-float USteamVRInputDeviceFunctionLibrary::GetSteamVR_GlobalPredictedSecondsFromNow()
+float UDEPRECATED_USteamVRInputDeviceFunctionLibrary::GetSteamVR_GlobalPredictedSecondsFromNow()
 {
 	if (VRSystem() && VRInput())
 	{
@@ -578,7 +580,7 @@ float USteamVRInputDeviceFunctionLibrary::GetSteamVR_GlobalPredictedSecondsFromN
 	return -9999.f;
 }
 
-float USteamVRInputDeviceFunctionLibrary::SetSteamVR_GlobalPredictedSecondsFromNow(float NewValue)
+float UDEPRECATED_USteamVRInputDeviceFunctionLibrary::SetSteamVR_GlobalPredictedSecondsFromNow(float NewValue)
 {
 	if (VRSystem() && VRInput())
 	{
@@ -593,13 +595,13 @@ float USteamVRInputDeviceFunctionLibrary::SetSteamVR_GlobalPredictedSecondsFromN
 	return -9999.f;
 }
 
-void USteamVRInputDeviceFunctionLibrary::ShowAllSteamVR_ActionOrigins()
+void UDEPRECATED_USteamVRInputDeviceFunctionLibrary::ShowAllSteamVR_ActionOrigins()
 {
 	VRActiveActionSet_t ActiveActionSet;
 	VRInput()->ShowBindingsForActionSet(&ActiveActionSet, sizeof(ActiveActionSet), 1, 0);
 }
 
-TArray<FSteamVRInputBindingInfo> USteamVRInputDeviceFunctionLibrary::GetSteamVR_InputBindingInfo(FSteamVRAction SteamVRActionHandle)
+TArray<FSteamVRInputBindingInfo> UDEPRECATED_USteamVRInputDeviceFunctionLibrary::GetSteamVR_InputBindingInfo(FSteamVRAction SteamVRActionHandle)
 {
 	// Setup variables to be used by the OpenVR call
 	TArray<FSteamVRInputBindingInfo> SteamVRBindingInfo;
@@ -657,7 +659,7 @@ TArray<FSteamVRInputBindingInfo> USteamVRInputDeviceFunctionLibrary::GetSteamVR_
 	return SteamVRBindingInfo;
 }
 
-TArray<FSteamVRInputBindingInfo> USteamVRInputDeviceFunctionLibrary::FindSteamVR_InputBindingInfo(FName ActionName, FName ActionSet /*= FName("main")*/)
+TArray<FSteamVRInputBindingInfo> UDEPRECATED_USteamVRInputDeviceFunctionLibrary::FindSteamVR_InputBindingInfo(FName ActionName, FName ActionSet /*= FName("main")*/)
 {
 	// Check for a valid OpenVR session and action name/set
 	if (VRInput() && ActionName != NAME_None && ActionSet != NAME_None)
@@ -692,7 +694,7 @@ TArray<FSteamVRInputBindingInfo> USteamVRInputDeviceFunctionLibrary::FindSteamVR
 
 }
 
-bool USteamVRInputDeviceFunctionLibrary::ResetSeatedPosition()
+bool UDEPRECATED_USteamVRInputDeviceFunctionLibrary::ResetSeatedPosition()
 {
 	if (VRSystem() && VRInput())
 	{
@@ -711,7 +713,7 @@ bool USteamVRInputDeviceFunctionLibrary::ResetSeatedPosition()
 }
 
 
-float USteamVRInputDeviceFunctionLibrary::GetUserIPD()
+float UDEPRECATED_USteamVRInputDeviceFunctionLibrary::GetUserIPD()
 {
 	if (VRSystem())
 	{
@@ -721,7 +723,7 @@ float USteamVRInputDeviceFunctionLibrary::GetUserIPD()
 	return 0.f;
 }
 
-FTransform USteamVRInputDeviceFunctionLibrary::GetUETransform(VRBoneTransform_t SteamBoneTransform)
+FTransform UDEPRECATED_USteamVRInputDeviceFunctionLibrary::GetUETransform(VRBoneTransform_t SteamBoneTransform)
 {
 	FTransform RetTransform;
 
@@ -738,7 +740,7 @@ FTransform USteamVRInputDeviceFunctionLibrary::GetUETransform(VRBoneTransform_t
 	return RetTransform;
 }
 
-void USteamVRInputDeviceFunctionLibrary::LaunchBindingsURL()
+void UDEPRECATED_USteamVRInputDeviceFunctionLibrary::LaunchBindingsURL()
 {
 	FSteamVRInputDevice* SteamVRInputDevice = GetSteamVRInputDevice();
 	if (SteamVRInputDevice != nullptr)
@@ -749,7 +751,7 @@ void USteamVRInputDeviceFunctionLibrary::LaunchBindingsURL()
 	}
 }
 
-void USteamVRInputDeviceFunctionLibrary::GetFingerCurlsAndSplays(EHand Hand, FSteamVRFingerCurls& FingerCurls, FSteamVRFingerSplays& FingerSplays, ESkeletalSummaryDataType SummaryDataType)
+void UDEPRECATED_USteamVRInputDeviceFunctionLibrary::GetFingerCurlsAndSplays(EHand Hand, FSteamVRFingerCurls& FingerCurls, FSteamVRFingerSplays& FingerSplays, ESkeletalSummaryDataType SummaryDataType)
 {
 	FSteamVRInputDevice* SteamVRInputDevice = GetSteamVRInputDevice();
 	if (SteamVRInputDevice != nullptr && VRSystem() && VRInput())
@@ -837,7 +839,7 @@ void USteamVRInputDeviceFunctionLibrary::GetFingerCurlsAndSplays(EHand Hand, FSt
 	FingerSplays = {};
 }
 
-FSteamVRInputDevice* USteamVRInputDeviceFunctionLibrary::GetSteamVRInputDevice()
+FSteamVRInputDevice* UDEPRECATED_USteamVRInputDeviceFunctionLibrary::GetSteamVRInputDevice()
 {
 	TArray<IMotionController*> MotionControllers = IModularFeatures::Get().GetModularFeatureImplementations<IMotionController>(IMotionController::GetModularFeatureName());
 	for (auto MotionController : MotionControllers)
@@ -863,3 +865,4 @@ FSteamVRInputDevice* USteamVRInputDeviceFunctionLibrary::GetSteamVRInputDevice()
 	return nullptr;
 }
 
+PRAGMA_ENABLE_DEPRECATION_WARNINGS
diff --git a/Engine/Plugins/Runtime/Steam/SteamVR/Source/SteamVRInputDevice/Private/SteamVRInputDeviceModule.cpp b/Engine/Plugins/Runtime/Steam/SteamVR/Source/SteamVRInputDevice/Private/SteamVRInputDeviceModule.cpp
--- a/Engine/Plugins/Runtime/Steam/SteamVR/Source/SteamVRInputDevice/Private/SteamVRInputDeviceModule.cpp
+++ b/Engine/Plugins/Runtime/Steam/SteamVR/Source/SteamVRInputDevice/Private/SteamVRInputDeviceModule.cpp
@@ -41,6 +41,8 @@ POSSIBILITY OF SUCH DAMAGE.
 #include "IVREditorModule.h"
 #endif
 
+PRAGMA_DISABLE_DEPRECATION_WARNINGS
+
 class FSteamVRInputDeviceModule : public ISteamVRInputDeviceModule
 {
 	/* Creates a new instance of SteamVR Input Controller **/
@@ -90,3 +92,5 @@ void FSteamVRInputDeviceModule::ShutdownModule()
 }
 
 IMPLEMENT_MODULE(FSteamVRInputDeviceModule, SteamVRInputDevice)
+
+PRAGMA_ENABLE_DEPRECATION_WARNINGS
\ No newline at end of file
diff --git a/Engine/Plugins/Runtime/Steam/SteamVR/Source/SteamVRInputDevice/Private/SteamVRSkeletonDefinition.cpp b/Engine/Plugins/Runtime/Steam/SteamVR/Source/SteamVRInputDevice/Private/SteamVRSkeletonDefinition.cpp
--- a/Engine/Plugins/Runtime/Steam/SteamVR/Source/SteamVRInputDevice/Private/SteamVRSkeletonDefinition.cpp
+++ b/Engine/Plugins/Runtime/Steam/SteamVR/Source/SteamVRInputDevice/Private/SteamVRSkeletonDefinition.cpp
@@ -30,6 +30,8 @@ POSSIBILITY OF SUCH DAMAGE.
 
 #include "SteamVRSkeletonDefinition.h"
 
+PRAGMA_DISABLE_DEPRECATION_WARNINGS
+
 namespace SteamVRSkeleton
 {
 	// Define the index of each bone's parent
@@ -159,3 +161,5 @@ namespace SteamVRSkeleton
 		return kBoneChildList[nBoneIndex][nChildIndex];
 	}
 }
+
+PRAGMA_ENABLE_DEPRECATION_WARNINGS
diff --git a/Engine/Plugins/Runtime/Steam/SteamVR/Source/SteamVRInputDevice/Private/SteamVRTrackingRefComponent.cpp b/Engine/Plugins/Runtime/Steam/SteamVR/Source/SteamVRInputDevice/Private/SteamVRTrackingRefComponent.cpp
--- a/Engine/Plugins/Runtime/Steam/SteamVR/Source/SteamVRInputDevice/Private/SteamVRTrackingRefComponent.cpp
+++ b/Engine/Plugins/Runtime/Steam/SteamVR/Source/SteamVRInputDevice/Private/SteamVRTrackingRefComponent.cpp
@@ -35,12 +35,14 @@ POSSIBILITY OF SUCH DAMAGE.
 using namespace vr;
 DEFINE_LOG_CATEGORY_STATIC(LogSteamVRTrackingRefComponent, Log, All);
 
-USteamVRTrackingReferences::USteamVRTrackingReferences()
+PRAGMA_DISABLE_DEPRECATION_WARNINGS
+
+UDEPRECATED_USteamVRTrackingReferences::UDEPRECATED_USteamVRTrackingReferences()
 {
 	PrimaryComponentTick.bCanEverTick = true;
 }
 
-bool USteamVRTrackingReferences::ShowTrackingReferences(UStaticMesh* TrackingReferenceMesh)
+bool UDEPRECATED_USteamVRTrackingReferences::ShowTrackingReferences(UStaticMesh* TrackingReferenceMesh)
 {
 	if (!TrackingReferenceMesh->IsValidLowLevel())
 	{
@@ -124,7 +126,7 @@ bool USteamVRTrackingReferences::ShowTrackingReferences(UStaticMesh* TrackingRef
 	return false;
 }
 
-void USteamVRTrackingReferences::HideTrackingReferences()
+void UDEPRECATED_USteamVRTrackingReferences::HideTrackingReferences()
 {
 	// Check if need to clear up any existing Tracking Reference in-world
 	if (TrackingReferences.Num()> 0)
@@ -143,7 +145,7 @@ void USteamVRTrackingReferences::HideTrackingReferences()
 	}
 }
 
-void USteamVRTrackingReferences::BeginPlay()
+void UDEPRECATED_USteamVRTrackingReferences::BeginPlay()
 {
 	Super::BeginPlay();
 
@@ -157,7 +159,7 @@ void USteamVRTrackingReferences::BeginPlay()
 	}
 }
 
-void USteamVRTrackingReferences::TickComponent(float DeltaTime, ELevelTick TickType, FActorComponentTickFunction* ThisTickFunction)
+void UDEPRECATED_USteamVRTrackingReferences::TickComponent(float DeltaTime, ELevelTick TickType, FActorComponentTickFunction* ThisTickFunction)
 {
 	Super::TickComponent(DeltaTime, TickType, ThisTickFunction);
 
@@ -256,7 +258,7 @@ void USteamVRTrackingReferences::TickComponent(float DeltaTime, ELevelTick TickT
 	}
 }
 
-FName USteamVRTrackingReferences::GetDeviceClass(unsigned int id)
+FName UDEPRECATED_USteamVRTrackingReferences::GetDeviceClass(unsigned int id)
 {
 	// Get device class
 	ETrackedDeviceClass TrackedDeviceClass = VRSystem()->GetTrackedDeviceClass(id);
@@ -282,7 +284,7 @@ FName USteamVRTrackingReferences::GetDeviceClass(unsigned int id)
 	return DeviceClass;
 }
 
-bool USteamVRTrackingReferences::FindTrackedDevice(unsigned int id)
+bool UDEPRECATED_USteamVRTrackingReferences::FindTrackedDevice(unsigned int id)
 {
 	for (FActiveTrackedDevice Device : ActiveTrackingDevices)
 	{
@@ -293,3 +295,5 @@ bool USteamVRTrackingReferences::FindTrackedDevice(unsigned int id)
 	}
 	return false;
 }
+
+PRAGMA_ENABLE_DEPRECATION_WARNINGS
diff --git a/Engine/Plugins/Runtime/Steam/SteamVR/Source/SteamVRInputDevice/Public/AnimNode_SteamVRInputAnimPose.h b/Engine/Plugins/Runtime/Steam/SteamVR/Source/SteamVRInputDevice/Public/AnimNode_SteamVRInputAnimPose.h
--- a/Engine/Plugins/Runtime/Steam/SteamVR/Source/SteamVRInputDevice/Public/AnimNode_SteamVRInputAnimPose.h
+++ b/Engine/Plugins/Runtime/Steam/SteamVR/Source/SteamVRInputDevice/Public/AnimNode_SteamVRInputAnimPose.h
@@ -36,36 +36,44 @@ POSSIBILITY OF SUCH DAMAGE.
 #include "SteamVRSkeletonDefinition.h"
 #include "AnimNode_SteamVRInputAnimPose.generated.h"
 
+PRAGMA_DISABLE_DEPRECATION_WARNINGS
+
 /**
 * Custom animation node to retrieve poses from the Skeletal Input System
 */
-USTRUCT(BlueprintType)
+USTRUCT(BlueprintType, meta = (Deprecated = "5.1"))
 struct STEAMVRINPUTDEVICE_API FAnimNode_SteamVRInputAnimPose : public FAnimNode_Base
 {
 	GENERATED_USTRUCT_BODY()
 
 	/** Range of motion for the skeletal input values */
-	UPROPERTY(EditAnywhere, BlueprintReadWrite, Category = Settings, meta = (AlwaysAsPin))
+	UE_DEPRECATED(5.1, "SteamVR plugin is deprecated; please use the OpenXR plugin.")
+	UPROPERTY(EditAnywhere, BlueprintReadWrite, Category = Settings, meta = (AlwaysAsPin, DeprecatedProperty))
 	EMotionRange MotionRange = EMotionRange::VR_WithoutController;
 
 	/** Which hand should the animation node retrieve skeletal input values for */
-	UPROPERTY(EditAnywhere, BlueprintReadWrite, Category = Settings, meta = (AlwaysAsPin))
+	UE_DEPRECATED(5.1, "SteamVR plugin is deprecated; please use the OpenXR plugin.")
+	UPROPERTY(EditAnywhere, BlueprintReadWrite, Category = Settings, meta = (AlwaysAsPin, DeprecatedProperty))
 	EHand Hand = EHand::VR_LeftHand;
 
 	/** What kind of skeleton are we dealing with */
-	UPROPERTY(EditAnywhere, BlueprintReadWrite, Category = Settings, meta = (AlwaysAsPin))
+	UE_DEPRECATED(5.1, "SteamVR plugin is deprecated; please use the OpenXR plugin.")
+	UPROPERTY(EditAnywhere, BlueprintReadWrite, Category = Settings, meta = (AlwaysAsPin, DeprecatedProperty))
 	EHandSkeleton HandSkeleton = EHandSkeleton::VR_SteamVRHandSkeleton;
 
 	/** Should the pose be mirrored so it can be applied to the opposite hand */
-	UPROPERTY(EditAnywhere, BlueprintReadWrite, Category = Settings, meta = ( AlwaysAsPin ))
+	UE_DEPRECATED(5.1, "SteamVR plugin is deprecated; please use the OpenXR plugin.")
+	UPROPERTY(EditAnywhere, BlueprintReadWrite, Category = Settings, meta = (AlwaysAsPin, DeprecatedProperty))
 	bool Mirror = false;
 
 	/** The UE4 equivalent of the SteamVR Transform values per bone */
-	UPROPERTY(EditAnywhere, BlueprintReadWrite, Category = SteamVRInput)
+	UE_DEPRECATED(5.1, "SteamVR plugin is deprecated; please use the OpenXR plugin.")
+	UPROPERTY(EditAnywhere, BlueprintReadWrite, Category = SteamVRInput, meta = (DeprecatedProperty))
 	FSteamVRSkeletonTransform SteamVRSkeletalTransform;
 
 	/** SteamVR Skeleton to UE4 retargetting cache */
-	UPROPERTY()
+	UE_DEPRECATED(5.1, "SteamVR plugin is deprecated; please use the OpenXR plugin.")
+	UPROPERTY(meta = (DeprecatedProperty))
 	FUE4RetargettingRefs UE4RetargettingRefs;
 
 public:
@@ -92,3 +100,5 @@ public:
 	FTransform CalcModelSpaceTransform(const FCompactPose& Pose, FCompactPoseBoneIndex BoneIndex);
 
 };
+
+PRAGMA_ENABLE_DEPRECATION_WARNINGS
diff --git a/Engine/Plugins/Runtime/Steam/SteamVR/Source/SteamVRInputDevice/Public/AnimNode_SteamVRSetWristTransform.h b/Engine/Plugins/Runtime/Steam/SteamVR/Source/SteamVRInputDevice/Public/AnimNode_SteamVRSetWristTransform.h
--- a/Engine/Plugins/Runtime/Steam/SteamVR/Source/SteamVRInputDevice/Public/AnimNode_SteamVRSetWristTransform.h
+++ b/Engine/Plugins/Runtime/Steam/SteamVR/Source/SteamVRInputDevice/Public/AnimNode_SteamVRSetWristTransform.h
@@ -36,24 +36,29 @@ POSSIBILITY OF SUCH DAMAGE.
 #include "SteamVRSkeletonDefinition.h"
 #include "AnimNode_SteamVRSetWristTransform.generated.h"
 
+PRAGMA_DISABLE_DEPRECATION_WARNINGS
+
 /**
 * Custom animation node that sets the wrist transform of a target pose from a reference pose
 */
-USTRUCT(BlueprintInternalUseOnly)
+USTRUCT(BlueprintInternalUseOnly, meta = (Deprecated = "5.1"))
 struct STEAMVRINPUTDEVICE_API FAnimNode_SteamVRSetWristTransform : public FAnimNode_Base
 {
 	GENERATED_USTRUCT_BODY()
 
 	/** The pose from where we will get the root and/or wrist transform from */
-	UPROPERTY(EditAnywhere, BlueprintReadWrite, Category = Links)
+	UE_DEPRECATED(5.1, "SteamVR plugin is deprecated; please use the OpenXR plugin.")
+	UPROPERTY(EditAnywhere, BlueprintReadWrite, Category = Links, meta = (DeprecatedProperty))
 	FPoseLink ReferencePose;
 
 	/** What kind of skeleton is used in the reference pose */
-	UPROPERTY(EditAnywhere, BlueprintReadWrite, Category = Settings, meta = (AlwaysAsPin))
+	UE_DEPRECATED(5.1, "SteamVR plugin is deprecated; please use the OpenXR plugin.")
+	UPROPERTY(EditAnywhere, BlueprintReadWrite, Category = Settings, meta = (AlwaysAsPin, DeprecatedProperty))
 	EHandSkeleton HandSkeleton = EHandSkeleton::VR_SteamVRHandSkeleton;
 
 	/** The pose to apply the wrist transform to */
-	UPROPERTY(EditAnywhere, BlueprintReadWrite, Category = Links)
+	UE_DEPRECATED(5.1, "SteamVR plugin is deprecated; please use the OpenXR plugin.")
+	UPROPERTY(EditAnywhere, BlueprintReadWrite, Category = Links, meta = (DeprecatedProperty))
 	FPoseLink TargetPose;
 
 public:
@@ -75,3 +80,5 @@ private:
 	FCompactPoseBoneIndex SteamVRWristBoneIndex = FCompactPoseBoneIndex(1);
 
 };
+
+PRAGMA_ENABLE_DEPRECATION_WARNINGS
diff --git a/Engine/Plugins/Runtime/Steam/SteamVR/Source/SteamVRInputDevice/Public/ISteamVRInputDeviceModule.h b/Engine/Plugins/Runtime/Steam/SteamVR/Source/SteamVRInputDevice/Public/ISteamVRInputDeviceModule.h
--- a/Engine/Plugins/Runtime/Steam/SteamVR/Source/SteamVRInputDevice/Public/ISteamVRInputDeviceModule.h
+++ b/Engine/Plugins/Runtime/Steam/SteamVR/Source/SteamVRInputDevice/Public/ISteamVRInputDeviceModule.h
@@ -34,6 +34,8 @@ POSSIBILITY OF SUCH DAMAGE.
 #include "Modules/ModuleManager.h"
 #include "IInputDeviceModule.h"
 
+class UE_DEPRECATED(5.1, "SteamVR plugin is deprecated; please use the OpenXR plugin.") ISteamVRInputDeviceModule;
+
 class STEAMVRINPUTDEVICE_API ISteamVRInputDeviceModule : public IInputDeviceModule
 {
 
diff --git a/Engine/Plugins/Runtime/Steam/SteamVR/Source/SteamVRInputDevice/Public/SteamVRInputDeviceFunctionLibrary.h b/Engine/Plugins/Runtime/Steam/SteamVR/Source/SteamVRInputDevice/Public/SteamVRInputDeviceFunctionLibrary.h
--- a/Engine/Plugins/Runtime/Steam/SteamVR/Source/SteamVRInputDevice/Public/SteamVRInputDeviceFunctionLibrary.h
+++ b/Engine/Plugins/Runtime/Steam/SteamVR/Source/SteamVRInputDevice/Public/SteamVRInputDeviceFunctionLibrary.h
@@ -48,135 +48,179 @@ using namespace vr;
 #define ACTION_PATH_VIBRATE_RIGHT		"/actions/main/out/vibrateright"
 #define MAX_BINDINGINFO_COUNT	5
 
+PRAGMA_DISABLE_DEPRECATION_WARNINGS
+
 /** UE4 Bone definition of the SteamVR Skeleton */
-USTRUCT(BlueprintType)
+USTRUCT(BlueprintType, meta = (Deprecated = "5.1"))
 struct STEAMVRINPUTDEVICE_API FSteamVRSkeletonTransform
 {
 	GENERATED_BODY()
 
-	UPROPERTY(EditAnywhere, BlueprintReadOnly, Category = "SteamVR Input")
+	UE_DEPRECATED(5.1, "SteamVR plugin is deprecated; please use the OpenXR plugin.")
+	UPROPERTY(EditAnywhere, BlueprintReadOnly, Category = "SteamVR Input", meta = (DeprecatedProperty))
 	FTransform Root;
-	UPROPERTY(EditAnywhere, BlueprintReadOnly, Category = "SteamVR Input")
+	UE_DEPRECATED(5.1, "SteamVR plugin is deprecated; please use the OpenXR plugin.")
+	UPROPERTY(EditAnywhere, BlueprintReadOnly, Category = "SteamVR Input", meta = (DeprecatedProperty))
 	FTransform Wrist;
 
-	UPROPERTY(EditAnywhere, BlueprintReadOnly, Category = "SteamVR Input")
+	UE_DEPRECATED(5.1, "SteamVR plugin is deprecated; please use the OpenXR plugin.")
+	UPROPERTY(EditAnywhere, BlueprintReadOnly, Category = "SteamVR Input", meta = (DeprecatedProperty))
 	FTransform Thumb_0;
-	UPROPERTY(EditAnywhere, BlueprintReadOnly, Category = "SteamVR Input")
+	UE_DEPRECATED(5.1, "SteamVR plugin is deprecated; please use the OpenXR plugin.")
+	UPROPERTY(EditAnywhere, BlueprintReadOnly, Category = "SteamVR Input", meta = (DeprecatedProperty))
 	FTransform Thumb_1;
-	UPROPERTY(EditAnywhere, BlueprintReadOnly, Category = "SteamVR Input")
+	UE_DEPRECATED(5.1, "SteamVR plugin is deprecated; please use the OpenXR plugin.")
+	UPROPERTY(EditAnywhere, BlueprintReadOnly, Category = "SteamVR Input", meta = (DeprecatedProperty))
 	FTransform Thumb_2;
-	UPROPERTY(EditAnywhere, BlueprintReadOnly, Category = "SteamVR Input")
+	UE_DEPRECATED(5.1, "SteamVR plugin is deprecated; please use the OpenXR plugin.")
+	UPROPERTY(EditAnywhere, BlueprintReadOnly, Category = "SteamVR Input", meta = (DeprecatedProperty))
 	FTransform Thumb_3;
 	
-	UPROPERTY(EditAnywhere, BlueprintReadOnly, Category = "SteamVR Input")
+	UE_DEPRECATED(5.1, "SteamVR plugin is deprecated; please use the OpenXR plugin.")
+	UPROPERTY(EditAnywhere, BlueprintReadOnly, Category = "SteamVR Input", meta = (DeprecatedProperty))
 	FTransform Index_0;
-	UPROPERTY(EditAnywhere, BlueprintReadOnly, Category = "SteamVR Input")
+	UE_DEPRECATED(5.1, "SteamVR plugin is deprecated; please use the OpenXR plugin.")
+	UPROPERTY(EditAnywhere, BlueprintReadOnly, Category = "SteamVR Input", meta = (DeprecatedProperty))
 	FTransform Index_1;
-	UPROPERTY(EditAnywhere, BlueprintReadOnly, Category = "SteamVR Input")
+	UE_DEPRECATED(5.1, "SteamVR plugin is deprecated; please use the OpenXR plugin.")
+	UPROPERTY(EditAnywhere, BlueprintReadOnly, Category = "SteamVR Input", meta = (DeprecatedProperty))
 	FTransform Index_2;
-	UPROPERTY(EditAnywhere, BlueprintReadOnly, Category = "SteamVR Input")
+	UE_DEPRECATED(5.1, "SteamVR plugin is deprecated; please use the OpenXR plugin.")
+	UPROPERTY(EditAnywhere, BlueprintReadOnly, Category = "SteamVR Input", meta = (DeprecatedProperty))
 	FTransform Index_3;
-	UPROPERTY(EditAnywhere, BlueprintReadOnly, Category = "SteamVR Input")
+	UE_DEPRECATED(5.1, "SteamVR plugin is deprecated; please use the OpenXR plugin.")
+	UPROPERTY(EditAnywhere, BlueprintReadOnly, Category = "SteamVR Input", meta = (DeprecatedProperty))
 	FTransform Index_4;
 
-	UPROPERTY(EditAnywhere, BlueprintReadOnly, Category = "SteamVR Input")
+	UE_DEPRECATED(5.1, "SteamVR plugin is deprecated; please use the OpenXR plugin.")
+	UPROPERTY(EditAnywhere, BlueprintReadOnly, Category = "SteamVR Input", meta = (DeprecatedProperty))
 	FTransform Middle_0;
-	UPROPERTY(EditAnywhere, BlueprintReadOnly, Category = "SteamVR Input")
+	UE_DEPRECATED(5.1, "SteamVR plugin is deprecated; please use the OpenXR plugin.")
+	UPROPERTY(EditAnywhere, BlueprintReadOnly, Category = "SteamVR Input", meta = (DeprecatedProperty))
 	FTransform Middle_1;
-	UPROPERTY(EditAnywhere, BlueprintReadOnly, Category = "SteamVR Input")
+	UE_DEPRECATED(5.1, "SteamVR plugin is deprecated; please use the OpenXR plugin.")
+	UPROPERTY(EditAnywhere, BlueprintReadOnly, Category = "SteamVR Input", meta = (DeprecatedProperty))
 	FTransform Middle_2;
-	UPROPERTY(EditAnywhere, BlueprintReadOnly, Category = "SteamVR Input")
+	UE_DEPRECATED(5.1, "SteamVR plugin is deprecated; please use the OpenXR plugin.")
+	UPROPERTY(EditAnywhere, BlueprintReadOnly, Category = "SteamVR Input", meta = (DeprecatedProperty))
 	FTransform Middle_3;
-	UPROPERTY(EditAnywhere, BlueprintReadOnly, Category = "SteamVR Input")
+	UE_DEPRECATED(5.1, "SteamVR plugin is deprecated; please use the OpenXR plugin.")
+	UPROPERTY(EditAnywhere, BlueprintReadOnly, Category = "SteamVR Input", meta = (DeprecatedProperty))
 	FTransform Middle_4;
 
-	UPROPERTY(EditAnywhere, BlueprintReadOnly, Category = "SteamVR Input")
+	UE_DEPRECATED(5.1, "SteamVR plugin is deprecated; please use the OpenXR plugin.")
+	UPROPERTY(EditAnywhere, BlueprintReadOnly, Category = "SteamVR Input", meta = (DeprecatedProperty))
 	FTransform Ring_0;
-	UPROPERTY(EditAnywhere, BlueprintReadOnly, Category = "SteamVR Input")
+	UE_DEPRECATED(5.1, "SteamVR plugin is deprecated; please use the OpenXR plugin.")
+	UPROPERTY(EditAnywhere, BlueprintReadOnly, Category = "SteamVR Input", meta = (DeprecatedProperty))
 	FTransform Ring_1;
-	UPROPERTY(EditAnywhere, BlueprintReadOnly, Category = "SteamVR Input")
+	UE_DEPRECATED(5.1, "SteamVR plugin is deprecated; please use the OpenXR plugin.")
+	UPROPERTY(EditAnywhere, BlueprintReadOnly, Category = "SteamVR Input", meta = (DeprecatedProperty))
 	FTransform Ring_2;
-	UPROPERTY(EditAnywhere, BlueprintReadOnly, Category = "SteamVR Input")
+	UE_DEPRECATED(5.1, "SteamVR plugin is deprecated; please use the OpenXR plugin.")
+	UPROPERTY(EditAnywhere, BlueprintReadOnly, Category = "SteamVR Input", meta = (DeprecatedProperty))
 	FTransform Ring_3;
-	UPROPERTY(EditAnywhere, BlueprintReadOnly, Category = "SteamVR Input")
+	UE_DEPRECATED(5.1, "SteamVR plugin is deprecated; please use the OpenXR plugin.")
+	UPROPERTY(EditAnywhere, BlueprintReadOnly, Category = "SteamVR Input", meta = (DeprecatedProperty))
 	FTransform Ring_4;
 
-	UPROPERTY(EditAnywhere, BlueprintReadOnly, Category = "SteamVR Input")
+	UE_DEPRECATED(5.1, "SteamVR plugin is deprecated; please use the OpenXR plugin.")
+	UPROPERTY(EditAnywhere, BlueprintReadOnly, Category = "SteamVR Input", meta = (DeprecatedProperty))
 	FTransform Pinky_0;
-	UPROPERTY(EditAnywhere, BlueprintReadOnly, Category = "SteamVR Input")
+	UE_DEPRECATED(5.1, "SteamVR plugin is deprecated; please use the OpenXR plugin.")
+	UPROPERTY(EditAnywhere, BlueprintReadOnly, Category = "SteamVR Input", meta = (DeprecatedProperty))
 	FTransform Pinky_1;
-	UPROPERTY(EditAnywhere, BlueprintReadOnly, Category = "SteamVR Input")
+	UE_DEPRECATED(5.1, "SteamVR plugin is deprecated; please use the OpenXR plugin.")
+	UPROPERTY(EditAnywhere, BlueprintReadOnly, Category = "SteamVR Input", meta = (DeprecatedProperty))
 	FTransform Pinky_2;
-	UPROPERTY(EditAnywhere, BlueprintReadOnly, Category = "SteamVR Input")
+	UE_DEPRECATED(5.1, "SteamVR plugin is deprecated; please use the OpenXR plugin.")
+	UPROPERTY(EditAnywhere, BlueprintReadOnly, Category = "SteamVR Input", meta = (DeprecatedProperty))
 	FTransform Pinky_3;
-	UPROPERTY(EditAnywhere, BlueprintReadOnly, Category = "SteamVR Input")
+	UE_DEPRECATED(5.1, "SteamVR plugin is deprecated; please use the OpenXR plugin.")
+	UPROPERTY(EditAnywhere, BlueprintReadOnly, Category = "SteamVR Input", meta = (DeprecatedProperty))
 	FTransform Pinky_4;
 
-	UPROPERTY(EditAnywhere, BlueprintReadOnly, Category = "SteamVR Input")
+	UE_DEPRECATED(5.1, "SteamVR plugin is deprecated; please use the OpenXR plugin.")
+	UPROPERTY(EditAnywhere, BlueprintReadOnly, Category = "SteamVR Input", meta = (DeprecatedProperty))
 	FTransform Aux_Thumb;
-	UPROPERTY(EditAnywhere, BlueprintReadOnly, Category = "SteamVR Input")
+	UE_DEPRECATED(5.1, "SteamVR plugin is deprecated; please use the OpenXR plugin.")
+	UPROPERTY(EditAnywhere, BlueprintReadOnly, Category = "SteamVR Input", meta = (DeprecatedProperty))
 	FTransform Aux_Index;
-	UPROPERTY(EditAnywhere, BlueprintReadOnly, Category = "SteamVR Input")
+	UE_DEPRECATED(5.1, "SteamVR plugin is deprecated; please use the OpenXR plugin.")
+	UPROPERTY(EditAnywhere, BlueprintReadOnly, Category = "SteamVR Input", meta = (DeprecatedProperty))
 	FTransform Aux_Middle;
-	UPROPERTY(EditAnywhere, BlueprintReadOnly, Category = "SteamVR Input")
+	UE_DEPRECATED(5.1, "SteamVR plugin is deprecated; please use the OpenXR plugin.")
+	UPROPERTY(EditAnywhere, BlueprintReadOnly, Category = "SteamVR Input", meta = (DeprecatedProperty))
 	FTransform Aux_Ring;
-	UPROPERTY(EditAnywhere, BlueprintReadOnly, Category = "SteamVR Input")
+	UE_DEPRECATED(5.1, "SteamVR plugin is deprecated; please use the OpenXR plugin.")
+	UPROPERTY(EditAnywhere, BlueprintReadOnly, Category = "SteamVR Input", meta = (DeprecatedProperty))
 	FTransform Aux_Pinky;
 };
 
 /** SteamVR finger curls */
-USTRUCT(BlueprintType)
+USTRUCT(BlueprintType, meta = (Deprecated = "5.1"))
 struct STEAMVRINPUTDEVICE_API FSteamVRFingerCurls
 {
 	GENERATED_BODY()
 
-	UPROPERTY(EditAnywhere, BlueprintReadOnly, Category = "SteamVR Input")
+	UE_DEPRECATED(5.1, "SteamVR plugin is deprecated; please use the OpenXR plugin.")
+	UPROPERTY(EditAnywhere, BlueprintReadOnly, Category = "SteamVR Input", meta = (DeprecatedProperty))
 	float	Thumb = 0.0f;
 
-	UPROPERTY(EditAnywhere, BlueprintReadOnly, Category = "SteamVR Input")
+	UE_DEPRECATED(5.1, "SteamVR plugin is deprecated; please use the OpenXR plugin.")
+	UPROPERTY(EditAnywhere, BlueprintReadOnly, Category = "SteamVR Input", meta = (DeprecatedProperty))
 	float	Index = 0.0f;
 
-	UPROPERTY(EditAnywhere, BlueprintReadOnly, Category = "SteamVR Input")
+	UE_DEPRECATED(5.1, "SteamVR plugin is deprecated; please use the OpenXR plugin.")
+	UPROPERTY(EditAnywhere, BlueprintReadOnly, Category = "SteamVR Input", meta = (DeprecatedProperty))
 	float	Middle = 0.0f;
 	
-	UPROPERTY(EditAnywhere, BlueprintReadOnly, Category = "SteamVR Input")
+	UE_DEPRECATED(5.1, "SteamVR plugin is deprecated; please use the OpenXR plugin.")
+	UPROPERTY(EditAnywhere, BlueprintReadOnly, Category = "SteamVR Input", meta = (DeprecatedProperty))
 	float	Ring = 0.0f;
 
-	UPROPERTY(EditAnywhere, BlueprintReadOnly, Category = "SteamVR Input")
+	UE_DEPRECATED(5.1, "SteamVR plugin is deprecated; please use the OpenXR plugin.")
+	UPROPERTY(EditAnywhere, BlueprintReadOnly, Category = "SteamVR Input", meta = (DeprecatedProperty))
 	float	Pinky = 0.0f;
 };
 
 /** SteamVR finger splays */
-USTRUCT(BlueprintType)
+USTRUCT(BlueprintType, meta = (Deprecated = "5.1"))
 struct STEAMVRINPUTDEVICE_API FSteamVRFingerSplays
 {
 	GENERATED_BODY()
 
-	UPROPERTY(EditAnywhere, BlueprintReadOnly, Category = "SteamVR Input")
+	UE_DEPRECATED(5.1, "SteamVR plugin is deprecated; please use the OpenXR plugin.")
+	UPROPERTY(EditAnywhere, BlueprintReadOnly, Category = "SteamVR Input", meta = (DeprecatedProperty))
 	float	Thumb_Index = 0.0f;
 
-	UPROPERTY(EditAnywhere, BlueprintReadOnly, Category = "SteamVR Input")
+	UE_DEPRECATED(5.1, "SteamVR plugin is deprecated; please use the OpenXR plugin.")
+	UPROPERTY(EditAnywhere, BlueprintReadOnly, Category = "SteamVR Input", meta = (DeprecatedProperty))
 	float	Index_Middle = 0.0f;
 
-	UPROPERTY(EditAnywhere, BlueprintReadOnly, Category = "SteamVR Input")
+	UE_DEPRECATED(5.1, "SteamVR plugin is deprecated; please use the OpenXR plugin.")
+	UPROPERTY(EditAnywhere, BlueprintReadOnly, Category = "SteamVR Input", meta = (DeprecatedProperty))
 	float	Middle_Ring = 0.0f;
 
-	UPROPERTY(EditAnywhere, BlueprintReadOnly, Category = "SteamVR Input")
+	UE_DEPRECATED(5.1, "SteamVR plugin is deprecated; please use the OpenXR plugin.")
+	UPROPERTY(EditAnywhere, BlueprintReadOnly, Category = "SteamVR Input", meta = (DeprecatedProperty))
 	float	Ring_Pinky = 0.0f;
 };
 
 /** SteamVR actions as defined by the developer */
-USTRUCT(BlueprintType)
+USTRUCT(BlueprintType, meta = (Deprecated = "5.1"))
 struct STEAMVRINPUTDEVICE_API FSteamVRAction
 {
 	GENERATED_BODY()
 
 	// The SteamVR name of the action (e.g. Teleport, OpenConsole)
-	UPROPERTY(EditAnywhere, BlueprintReadOnly, Category = "SteamVR Input")
+	UE_DEPRECATED(5.1, "SteamVR plugin is deprecated; please use the OpenXR plugin.")
+	UPROPERTY(EditAnywhere, BlueprintReadOnly, Category = "SteamVR Input", meta = (DeprecatedProperty))
 	FName		Name;
 
 	// The path defined for the action (e.g. main/in/{ActionName})
-	UPROPERTY(EditAnywhere, BlueprintReadOnly, Category = "SteamVR Input")
+	UE_DEPRECATED(5.1, "SteamVR plugin is deprecated; please use the OpenXR plugin.")
+	UPROPERTY(EditAnywhere, BlueprintReadOnly, Category = "SteamVR Input", meta = (DeprecatedProperty))
 	FString		Path;
 
 	VRActionHandle_t Handle;				// The handle to the SteamVR Action 
@@ -195,13 +239,14 @@ struct STEAMVRINPUTDEVICE_API FSteamVRAction
 
 
 /** SteamVR action set as defined by the developer */
-USTRUCT(BlueprintType)
+USTRUCT(BlueprintType, meta = (Deprecated = "5.1"))
 struct STEAMVRINPUTDEVICE_API FSteamVRActionSet
 {
 	GENERATED_BODY()
 
 	// The path defined for this action set (e.g. /actions/main)
-	UPROPERTY(EditAnywhere, BlueprintReadOnly, Category = "SteamVR Input")
+	UE_DEPRECATED(5.1, "SteamVR plugin is deprecated; please use the OpenXR plugin.")
+	UPROPERTY(EditAnywhere, BlueprintReadOnly, Category = "SteamVR Input", meta = (DeprecatedProperty))
 	FString		Path;			
 	
 	VRActionSetHandle_t Handle;	// The handle to the SteamVR Action Set
@@ -216,21 +261,24 @@ struct STEAMVRINPUTDEVICE_API FSteamVRActionSet
 };
 
 /** Information about the tracked device associated from the input source */
-USTRUCT(BlueprintType)
+USTRUCT(BlueprintType, meta = (Deprecated = "5.1"))
 struct STEAMVRINPUTDEVICE_API FSteamVRInputOriginInfo
 {
 	GENERATED_BODY()
 
 	// The tracked device index for the device or k_unTrackedDeviceInvalid (0xFFFFFFFF)
-	UPROPERTY(EditAnywhere, BlueprintReadOnly, Category = "SteamVR Input")
+	UE_DEPRECATED(5.1, "SteamVR plugin is deprecated; please use the OpenXR plugin.")
+	UPROPERTY(EditAnywhere, BlueprintReadOnly, Category = "SteamVR Input", meta = (DeprecatedProperty))
 	int32	TrackedDeviceIndex;
 	
 	//  The name of the component of the tracked device's render model that represents this input source, or an empty string if there is no associated render model component.
-	UPROPERTY(EditAnywhere, BlueprintReadOnly, Category = "SteamVR Input")
+	UE_DEPRECATED(5.1, "SteamVR plugin is deprecated; please use the OpenXR plugin.")
+	UPROPERTY(EditAnywhere, BlueprintReadOnly, Category = "SteamVR Input", meta = (DeprecatedProperty))
 	FString	RenderModelComponentName;
 
 	//  The tracked device's model info
-	UPROPERTY(EditAnywhere, BlueprintReadOnly, Category = "SteamVR Input")
+	UE_DEPRECATED(5.1, "SteamVR plugin is deprecated; please use the OpenXR plugin.")
+	UPROPERTY(EditAnywhere, BlueprintReadOnly, Category = "SteamVR Input", meta = (DeprecatedProperty))
 	FString	TrackedDeviceModel;
 
 	FSteamVRInputOriginInfo(const int32 InDeviceIndex, const FString InRenderModelComponentName)
@@ -245,24 +293,29 @@ struct STEAMVRINPUTDEVICE_API FSteamVRInputOriginInfo
 };
 
 /** Retargetting information for the SteamVR skeleton to UE4 stock skeleton */
-USTRUCT(BlueprintType)
+USTRUCT(BlueprintType, meta = (Deprecated = "5.1"))
 struct STEAMVRINPUTDEVICE_API FUE4RetargettingRefs
 {
 	GENERATED_BODY()
 
-	UPROPERTY(BlueprintReadOnly, Category = "SteamVR Input")
+	UE_DEPRECATED(5.1, "SteamVR plugin is deprecated; please use the OpenXR plugin.")
+	UPROPERTY(BlueprintReadOnly, Category = "SteamVR Input", meta = (DeprecatedProperty))
 	bool bIsInitialized;
 
-	UPROPERTY(BlueprintReadOnly, Category = "SteamVR Input")
+	UE_DEPRECATED(5.1, "SteamVR plugin is deprecated; please use the OpenXR plugin.")
+	UPROPERTY(BlueprintReadOnly, Category = "SteamVR Input", meta = (DeprecatedProperty))
 	bool bIsRightHanded;
 
-	UPROPERTY(BlueprintReadOnly, Category = "SteamVR Input")
+	UE_DEPRECATED(5.1, "SteamVR plugin is deprecated; please use the OpenXR plugin.")
+	UPROPERTY(BlueprintReadOnly, Category = "SteamVR Input", meta = (DeprecatedProperty))
 	FVector KnuckleAverageMS_UE4;
 
-	UPROPERTY(BlueprintReadOnly, Category = "SteamVR Input")
+	UE_DEPRECATED(5.1, "SteamVR plugin is deprecated; please use the OpenXR plugin.")
+	UPROPERTY(BlueprintReadOnly, Category = "SteamVR Input", meta = (DeprecatedProperty))
 	FVector WristSideDirectionLS_UE4;
 
-	UPROPERTY(BlueprintReadOnly, Category = "SteamVR Input")
+	UE_DEPRECATED(5.1, "SteamVR plugin is deprecated; please use the OpenXR plugin.")
+	UPROPERTY(BlueprintReadOnly, Category = "SteamVR Input", meta = (DeprecatedProperty))
 	FVector WristForwardLS_UE4;
 
 	FUE4RetargettingRefs()
@@ -276,21 +329,25 @@ struct STEAMVRINPUTDEVICE_API FUE4RetargettingRefs
 };
 
 /** Information about the input bindings for an action on currently active controller (i.e device path, input path, mode, slot)  */
-USTRUCT(BlueprintType)
+USTRUCT(BlueprintType, meta = (Deprecated = "5.1"))
 struct STEAMVRINPUTDEVICE_API FSteamVRInputBindingInfo
 {
 	GENERATED_BODY()
 
-	UPROPERTY(BlueprintReadOnly, Category = "SteamVR Input")
+	UE_DEPRECATED(5.1, "SteamVR plugin is deprecated; please use the OpenXR plugin.")
+	UPROPERTY(BlueprintReadOnly, Category = "SteamVR Input", meta = (DeprecatedProperty))
 	FName	DevicePathName;
 
-	UPROPERTY(BlueprintReadOnly, Category = "SteamVR Input")
+	UE_DEPRECATED(5.1, "SteamVR plugin is deprecated; please use the OpenXR plugin.")
+	UPROPERTY(BlueprintReadOnly, Category = "SteamVR Input", meta = (DeprecatedProperty))
 	FName	InputPathName;
 
-	UPROPERTY(BlueprintReadOnly, Category = "SteamVR Input")
+	UE_DEPRECATED(5.1, "SteamVR plugin is deprecated; please use the OpenXR plugin.")
+	UPROPERTY(BlueprintReadOnly, Category = "SteamVR Input", meta = (DeprecatedProperty))
 	FName	ModeName;
 
-	UPROPERTY(BlueprintReadOnly, Category = "SteamVR Input")
+	UE_DEPRECATED(5.1, "SteamVR plugin is deprecated; please use the OpenXR plugin.")
+	UPROPERTY(BlueprintReadOnly, Category = "SteamVR Input", meta = (DeprecatedProperty))
 	FName	SlotName;
 
 	FSteamVRInputBindingInfo()
@@ -312,7 +369,7 @@ struct STEAMVRINPUTDEVICE_API FSteamVRInputBindingInfo
 
 /** Convenience type for SteamVR Hand designation (Left/Right) */
 UENUM(BlueprintType)	
-enum class ESteamVRHand : uint8
+enum class UE_DEPRECATED(5.1, "SteamVR plugin is deprecated; please use the OpenXR plugin.") ESteamVRHand : uint8
 {
 	VR_Left		UMETA(DisplayName = "Left"),
 	VR_Right	UMETA(DisplayName = "Right")
@@ -320,7 +377,7 @@ enum class ESteamVRHand : uint8
 
 /** Valid range of motion for a skeletal animation */
 UENUM(BlueprintType)
-enum class EMotionRange : uint8
+enum class UE_DEPRECATED(5.1, "SteamVR plugin is deprecated; please use the OpenXR plugin.") EMotionRange : uint8
 {
 	VR_WithoutController 	UMETA(DisplayName = "Without Controller"),
 	VR_WithController 		UMETA(DisplayName = "With Controller")
@@ -328,7 +385,7 @@ enum class EMotionRange : uint8
 
 /** Valid values for hands thats used for the Skeletal Input System calls */
 UENUM(BlueprintType)
-enum class EHand : uint8
+enum class UE_DEPRECATED(5.1, "SteamVR plugin is deprecated; please use the OpenXR plugin.") EHand : uint8
 {
 	VR_LeftHand 	UMETA(DisplayName = "Left Hand"),
 	VR_RightHand 	UMETA(DisplayName = "Right Hand")
@@ -336,7 +393,7 @@ enum class EHand : uint8
 
 /** Types of known skeletons that this animation node can handle */
 UENUM(BlueprintType)
-enum class EHandSkeleton : uint8
+enum class UE_DEPRECATED(5.1, "SteamVR plugin is deprecated; please use the OpenXR plugin.") EHandSkeleton : uint8
 {
 	VR_SteamVRHandSkeleton 	UMETA(DisplayName = "SteamVR Hand Skeleton"),
 	VR_UE4HandSkeleton 		UMETA(DisplayName = "UE4 Hand Skeleton")
@@ -344,7 +401,7 @@ enum class EHandSkeleton : uint8
 
 /** Skeletal Tracking Level of a controller */
 UENUM(BlueprintType)
-enum class EControllerFidelity : uint8
+enum class UE_DEPRECATED(5.1, "SteamVR plugin is deprecated; please use the OpenXR plugin.") EControllerFidelity : uint8
 {
 	VR_ControllerFidelity_Estimated 	UMETA(DisplayName = "Controller Fidelity Estimated"),
 	VR_ControllerFidelity_Full 			UMETA(DisplayName = "Controller Fidelity Full"),
@@ -353,7 +410,7 @@ enum class EControllerFidelity : uint8
 
 /** Input String Values for querying user hardware */
 UENUM(BlueprintType)
-enum class ESteamVRInputStringBits : uint8
+enum class UE_DEPRECATED(5.1, "SteamVR plugin is deprecated; please use the OpenXR plugin.") ESteamVRInputStringBits : uint8
 {
 	// Which hand the origin is in (e.g. "Left Hand")
 	VR_InputString_Hand					UMETA(DisplayName = "Hand"),
@@ -370,7 +427,7 @@ enum class ESteamVRInputStringBits : uint8
 
 /** Input String Values for querying user hardware */
 UENUM(BlueprintType)
-enum class ESkeletalSummaryDataType : uint8
+enum class UE_DEPRECATED(5.1, "SteamVR plugin is deprecated; please use the OpenXR plugin.") ESkeletalSummaryDataType : uint8
 {
 	// The data should match the animated transforms in the skeleton transforms. This data will probably be smoothed and may be more latent
 	VR_SummaryType_FromAnimation		UMETA(DisplayName = "From Animation"),
@@ -383,8 +440,8 @@ enum class ESkeletalSummaryDataType : uint8
  * SteamVR Input Extended Functions
  * Functions and properties defined here are safe for developer use
  */
-UCLASS()
-class STEAMVRINPUTDEVICE_API USteamVRInputDeviceFunctionLibrary : public UBlueprintFunctionLibrary
+UCLASS(deprecated, meta = (DeprecationMessage = "SteamVR plugin is deprecated; please use the OpenXR plugin."))
+class STEAMVRINPUTDEVICE_API UDEPRECATED_USteamVRInputDeviceFunctionLibrary : public UBlueprintFunctionLibrary
 {
 	GENERATED_BODY()
 
@@ -394,7 +451,8 @@ public:
 	* @param Hand - Which hand to get the finger curls and splay values for
 	* @param FingerCurls - Curl values for each finger pair this frame
 	*/
-	UFUNCTION(BlueprintCallable, Category = "SteamVR Input")
+	UE_DEPRECATED(5.1, "SteamVR plugin is deprecated; please use the OpenXR plugin.")
+	UFUNCTION(BlueprintCallable, Category = "SteamVR Input", meta = (DeprecatedFunction, DeprecationMessage = "The SteamVR plugin is deprecated."))
 	static void GetFingerCurlsAndSplays(EHand Hand, FSteamVRFingerCurls& FingerCurls, FSteamVRFingerSplays& FingerSplays, ESkeletalSummaryDataType SummaryDataType = ESkeletalSummaryDataType::VR_SummaryType_FromAnimation);
 
 	/**
@@ -405,7 +463,8 @@ public:
 	* @param Frequency - Frequency used in the haptic feedback
 	* @param Amplitude - Amplitude used in the haptic feedback
 	*/
-	UFUNCTION(BlueprintCallable, Category="SteamVR Input")
+	UE_DEPRECATED(5.1, "SteamVR plugin is deprecated; please use the OpenXR plugin.")
+	UFUNCTION(BlueprintCallable, Category="SteamVR Input", meta = (DeprecatedFunction, DeprecationMessage = "The SteamVR plugin is deprecated."))
 	static void PlaySteamVR_HapticFeedback(ESteamVRHand Hand, float StartSecondsFromNow, float DurationSeconds = 1.f,
 			float Frequency = 1.f, float Amplitude = 0.5f);
 
@@ -414,7 +473,8 @@ public:
 	* @return LeftHandState - Whether or not curls and splay values are being retrieved from the left hand
 	* @return RightHandState -  Whether or not curls and splay values are being retrieved from the right hand
 	*/
-	UFUNCTION(BlueprintCallable, Category="SteamVR Input")
+	UE_DEPRECATED(5.1, "SteamVR plugin is deprecated; please use the OpenXR plugin.")
+	UFUNCTION(BlueprintCallable, Category="SteamVR Input", meta = (DeprecatedFunction, DeprecationMessage = "The SteamVR plugin is deprecated."))
 	static void GetCurlsAndSplaysState(bool& LeftHandState, bool& RightHandState);
 
 	/**
@@ -422,7 +482,8 @@ public:
 	* @return LeftHandState - Whether or not the controller attached to the player's left hand have skeletal input support
 	* @return RightHandState -  Whether or not the controller attached to the player's right hand have skeletal input support
 	*/
-	UFUNCTION(BlueprintCallable, Category = "SteamVR Input")
+	UE_DEPRECATED(5.1, "SteamVR plugin is deprecated; please use the OpenXR plugin.")
+	UFUNCTION(BlueprintCallable, Category = "SteamVR Input", meta = (DeprecatedFunction, DeprecationMessage = "The SteamVR plugin is deprecated."))
 	static void GetSkeletalState(bool& LeftHandState, bool& RightHandState);
 
 	/** 
@@ -430,7 +491,8 @@ public:
 	* @return LeftControllerFidelity - The skeletal tracking level of the left controller
 	* @return RightControllerFidelity -  The skeletal tracking level of the right controller
 	*/
-	UFUNCTION(BlueprintCallable, Category = "SteamVR Input")
+	UE_DEPRECATED(5.1, "SteamVR plugin is deprecated; please use the OpenXR plugin.")
+	UFUNCTION(BlueprintCallable, Category = "SteamVR Input", meta = (DeprecatedFunction, DeprecationMessage = "The SteamVR plugin is deprecated."))
 	static void GetControllerFidelity(EControllerFidelity& LeftControllerFidelity, EControllerFidelity& RightControllerFidelity);
 
 	/**
@@ -438,21 +500,24 @@ public:
 	* @param NewLeftHandState - Whether or not curls and splay values will be retrieved for the left hand
 	* @param NewRightHandState -  Whether or not curls and splay values will be retrieved for the right hand
 	*/
-	UFUNCTION(BlueprintCallable, Category = "SteamVR Input")
+	UE_DEPRECATED(5.1, "SteamVR plugin is deprecated; please use the OpenXR plugin.")
+	UFUNCTION(BlueprintCallable, Category = "SteamVR Input", meta = (DeprecatedFunction, DeprecationMessage = "The SteamVR plugin is deprecated."))
 	static void SetCurlsAndSplaysState(bool NewLeftHandState, bool NewRightHandState);
 
 	/**
 	* Check whether we are using a skeleton pose or the raw controller pose for the orientation and position of the motion controller
 	* @param bUseSkeletonPose - Whether or not we are using the skeleton pose instead of the controller raw pose
 	*/
-	UFUNCTION(BlueprintCallable, Category = "SteamVR Input")
+	UE_DEPRECATED(5.1, "SteamVR plugin is deprecated; please use the OpenXR plugin.")
+	UFUNCTION(BlueprintCallable, Category = "SteamVR Input", meta = (DeprecatedFunction, DeprecationMessage = "The SteamVR plugin is deprecated."))
 	static void GetPoseSource(bool& bUsingSkeletonPose);
 
 	/**
 	* Set whether to use skeleton pose or the raw controller pose for the orientation and position of the motion controller
 	* @param bUseSkeletonPose - Whether or not to use the skeleton pose instead of the controller raw pose
 	*/
-	UFUNCTION(BlueprintCallable, Category = "SteamVR Input")
+	UE_DEPRECATED(5.1, "SteamVR plugin is deprecated; please use the OpenXR plugin.")
+	UFUNCTION(BlueprintCallable, Category = "SteamVR Input", meta = (DeprecatedFunction, DeprecationMessage = "The SteamVR plugin is deprecated."))
 	static void SetPoseSource(bool bUseSkeletonPose);
 
 	/**
@@ -462,7 +527,8 @@ public:
 	* @param bWithController - Whether or not retrieve skeletal input values with controller
 	* @param bXAxisForward - Whether or not the Skeleton has the X axis facing forward
 	*/
-	UFUNCTION(BlueprintCallable, Category = "SteamVR Input")
+	UE_DEPRECATED(5.1, "SteamVR plugin is deprecated; please use the OpenXR plugin.")
+	UFUNCTION(BlueprintCallable, Category = "SteamVR Input", meta = (DeprecatedFunction, DeprecationMessage = "The SteamVR plugin is deprecated."))
 	static void GetSkeletalTransform(FSteamVRSkeletonTransform& LeftHand, FSteamVRSkeletonTransform& RightHand, bool bWithController=false);
 	
 	/**
@@ -472,7 +538,8 @@ public:
 	* @return AngularVelocity - The angular velocity of the hand this frame
 	* @return Velocity - The velocity of the hand this frame
 	*/
-	UFUNCTION(BlueprintCallable, Category = "SteamVR Input")
+	UE_DEPRECATED(5.1, "SteamVR plugin is deprecated; please use the OpenXR plugin.")
+	UFUNCTION(BlueprintCallable, Category = "SteamVR Input", meta = (DeprecatedFunction, DeprecationMessage = "The SteamVR plugin is deprecated."))
 	static void GetLeftHandPoseData(FVector& Position, FRotator& Orientation, FVector& AngularVelocity, FVector& Velocity);
 
 	/**
@@ -482,14 +549,16 @@ public:
 	* @return AngularVelocity - The angular velocity of the hand this frame
 	* @return Velocity - The velocity of the hand this frame
 	*/
-	UFUNCTION(BlueprintCallable, Category = "SteamVR Input")
+	UE_DEPRECATED(5.1, "SteamVR plugin is deprecated; please use the OpenXR plugin.")
+	UFUNCTION(BlueprintCallable, Category = "SteamVR Input", meta = (DeprecatedFunction, DeprecationMessage = "The SteamVR plugin is deprecated."))
 	static void GetRightHandPoseData(FVector& Position, FRotator& Orientation, FVector& AngularVelocity, FVector& Velocity);
 
 	/**
 	* Retrieve the input actions for this project
 	* @return SteamVRActions - Input actions defined in this project
 	*/
-	UFUNCTION(BlueprintCallable, Category = "SteamVR Input")
+	UE_DEPRECATED(5.1, "SteamVR plugin is deprecated; please use the OpenXR plugin.")
+	UFUNCTION(BlueprintCallable, Category = "SteamVR Input", meta = (DeprecatedFunction, DeprecationMessage = "The SteamVR plugin is deprecated."))
 	static void GetSteamVR_ActionArray(TArray<FSteamVRAction>& SteamVRActions);
 
 	/**
@@ -500,14 +569,16 @@ public:
 	* @return FoundAction - The action if found
 	* @return FoundActionSet - The action set if found
 	*/
-	UFUNCTION(BlueprintCallable, Category = "SteamVR Input")
+	UE_DEPRECATED(5.1, "SteamVR plugin is deprecated; please use the OpenXR plugin.")
+	UFUNCTION(BlueprintCallable, Category = "SteamVR Input", meta = (DeprecatedFunction, DeprecationMessage = "The SteamVR plugin is deprecated."))
 	static void FindSteamVR_Action(FName ActionName, bool& bResult, FSteamVRAction& FoundAction, FSteamVRActionSet& FoundActionSet, FName ActionSet = FName("main"));
 
 	/**
 	* Retrieve the input action sets for this project
 	* @return SteamVRActionSets - Input action sets defined in this project
 	*/
-	UFUNCTION(BlueprintCallable, Category = "SteamVR Input")
+	UE_DEPRECATED(5.1, "SteamVR plugin is deprecated; please use the OpenXR plugin.")
+	UFUNCTION(BlueprintCallable, Category = "SteamVR Input", meta = (DeprecatedFunction, DeprecationMessage = "The SteamVR plugin is deprecated."))
 	static void GetSteamVR_ActionSetArray(TArray<FSteamVRActionSet>& SteamVRActionSets);
 
 	/**
@@ -516,7 +587,8 @@ public:
 	* @return InputOriginInfo - The origin info of the action
 	* @return bool - whether the operation is successful or not
 	*/
-	UFUNCTION(BlueprintCallable, Category = "SteamVR Input")
+	UE_DEPRECATED(5.1, "SteamVR plugin is deprecated; please use the OpenXR plugin.")
+	UFUNCTION(BlueprintCallable, Category = "SteamVR Input", meta = (DeprecatedFunction, DeprecationMessage = "The SteamVR plugin is deprecated."))
 	static bool GetSteamVR_OriginTrackedDeviceInfo(FSteamVRAction SteamVRAction, FSteamVRInputOriginInfo& InputOriginInfo);
 
 	/**
@@ -525,7 +597,8 @@ public:
 	* @return InputOriginInfo - The origin info of the action
 	* @return bool - whether the operation is successful or not
 	*/
-	UFUNCTION(BlueprintCallable, Category = "SteamVR Input")
+	UE_DEPRECATED(5.1, "SteamVR plugin is deprecated; please use the OpenXR plugin.")
+	UFUNCTION(BlueprintCallable, Category = "SteamVR Input", meta = (DeprecatedFunction, DeprecationMessage = "The SteamVR plugin is deprecated."))
 	static void FindSteamVR_OriginTrackedDeviceInfo(FName ActionName, bool& bResult, FSteamVRInputOriginInfo& InputOriginInfo, FName ActionSet = FName("main"));
 
 	/**
@@ -534,7 +607,8 @@ public:
 	* @param LocalizedParts - Bitfields to specify which origin parts to return
 	* @return OriginLocalizedName -  The localized name of the origin of a given action (e.g. "Left Hand Index Controller Trackpad")
 	*/
-	UFUNCTION(BlueprintCallable, Category = "SteamVR Input")
+	UE_DEPRECATED(5.1, "SteamVR plugin is deprecated; please use the OpenXR plugin.")
+	UFUNCTION(BlueprintCallable, Category = "SteamVR Input", meta = (DeprecatedFunction, DeprecationMessage = "The SteamVR plugin is deprecated."))
 	static void GetSteamVR_OriginLocalizedName(FSteamVRAction SteamVRAction, TArray<ESteamVRInputStringBits> LocalizedParts, FString& OriginLocalizedName);
 
 	/**
@@ -542,7 +616,8 @@ public:
 	* @param SteamVRAction - The action that we will lookup the current binding for
 	* @param SteamVRActionSet - The action set that the action belongs to
 	*/
-	UFUNCTION(BlueprintCallable, Category = "SteamVR Input")
+	UE_DEPRECATED(5.1, "SteamVR plugin is deprecated; please use the OpenXR plugin.")
+	UFUNCTION(BlueprintCallable, Category = "SteamVR Input", meta = (DeprecatedFunction, DeprecationMessage = "The SteamVR plugin is deprecated."))
 	static void ShowSteamVR_ActionOrigin(FSteamVRAction SteamVRAction, FSteamVRActionSet SteamVRActionSet);
 
 	/**
@@ -550,7 +625,8 @@ public:
 	* @param SteamVRAction - The action that we will lookup the current binding for
 	* @param SteamVRActionSet - The action set that the action belongs to. Defaults to "main"
 	*/
-	UFUNCTION(BlueprintCallable, Category = "SteamVR Input")
+	UE_DEPRECATED(5.1, "SteamVR plugin is deprecated; please use the OpenXR plugin.")
+	UFUNCTION(BlueprintCallable, Category = "SteamVR Input", meta = (DeprecatedFunction, DeprecationMessage = "The SteamVR plugin is deprecated."))
 	static bool FindSteamVR_ActionOrigin(FName ActionName, FName ActionSet = FName("main"));
 
 	/**
@@ -560,7 +636,8 @@ public:
 	* @return Orientation - The rotation of the hand at the point in time, given a relative number of seconds, from the current time
 	* @return bool - Whether or not the call was succesful
 	*/
-	UFUNCTION(BlueprintCallable, Category = "SteamVR Input")
+	UE_DEPRECATED(5.1, "SteamVR plugin is deprecated; please use the OpenXR plugin.")
+	UFUNCTION(BlueprintCallable, Category = "SteamVR Input", meta = (DeprecatedFunction, DeprecationMessage = "The SteamVR plugin is deprecated."))
 	static bool GetSteamVR_HandPoseRelativeToNow(FVector& Position, FRotator& Orientation, ESteamVRHand Hand = ESteamVRHand::VR_Left, float PredictedSecondsFromNow = 0.f);
 
 	/**
@@ -568,7 +645,8 @@ public:
 	* A value of -9999.f triggers a GetPoseActionDataForNextFrame, otherwise GetPoseActionRelativeToNow is called with this value
 	* @return float - The current Predicted Seconds From Now from the SteamVRInput device
 	*/
-	UFUNCTION(BlueprintCallable, Category = "SteamVR Input")
+	UE_DEPRECATED(5.1, "SteamVR plugin is deprecated; please use the OpenXR plugin.")
+	UFUNCTION(BlueprintCallable, Category = "SteamVR Input", meta = (DeprecatedFunction, DeprecationMessage = "The SteamVR plugin is deprecated."))
 	static float GetSteamVR_GlobalPredictedSecondsFromNow();
 
 	/**
@@ -577,13 +655,15 @@ public:
 	* @param NewValue - The value for PredictedSecondsFromNow that will be used by the SteamVRInput device for Get Action Pose Data calls 
 	* @return float - The current Predicted Seconds From Now from the SteamVRInput device
 	*/
-	UFUNCTION(BlueprintCallable, Category = "SteamVR Input")
+	UE_DEPRECATED(5.1, "SteamVR plugin is deprecated; please use the OpenXR plugin.")
+	UFUNCTION(BlueprintCallable, Category = "SteamVR Input", meta = (DeprecatedFunction, DeprecationMessage = "The SteamVR plugin is deprecated."))
 	static float SetSteamVR_GlobalPredictedSecondsFromNow(float NewValue);
 
 	/**
 	* Shows all current bindings for the current controller in the user's headset
 	*/
-	UFUNCTION(BlueprintCallable, Category = "SteamVR Input")
+	UE_DEPRECATED(5.1, "SteamVR plugin is deprecated; please use the OpenXR plugin.")
+	UFUNCTION(BlueprintCallable, Category = "SteamVR Input", meta = (DeprecatedFunction, DeprecationMessage = "The SteamVR plugin is deprecated."))
 	static void ShowAllSteamVR_ActionOrigins();
 
 	/**
@@ -591,7 +671,8 @@ public:
 	* @param SteamVRActionHandle - The action handle of the action that binding info will be retrieved for the currently active controller. Use Find SteamVRAction node to get a handle
 	* @return SteamVRInputBindingInfo - Array of binding info for an action with the currently active controller
 	*/
-	UFUNCTION(BlueprintCallable, Category = "SteamVR Input")
+	UE_DEPRECATED(5.1, "SteamVR plugin is deprecated; please use the OpenXR plugin.")
+	UFUNCTION(BlueprintCallable, Category = "SteamVR Input", meta = (DeprecatedFunction, DeprecationMessage = "The SteamVR plugin is deprecated."))
 	static TArray<FSteamVRInputBindingInfo> GetSteamVR_InputBindingInfo(FSteamVRAction SteamVRActionHandle);
 
 	/**
@@ -600,7 +681,8 @@ public:
 	* @param ActionSet - The name of the action set that the action belongs in
 	* @return SteamVRInputBindingInfo - Array of binding info for an action with the currently active controller
 	*/
-	UFUNCTION(BlueprintCallable, Category = "SteamVR Input")
+	UE_DEPRECATED(5.1, "SteamVR plugin is deprecated; please use the OpenXR plugin.")
+	UFUNCTION(BlueprintCallable, Category = "SteamVR Input", meta = (DeprecatedFunction, DeprecationMessage = "The SteamVR plugin is deprecated."))
 	static TArray<FSteamVRInputBindingInfo> FindSteamVR_InputBindingInfo(FName ActionName, FName ActionSet = FName("main"));
 
 	/**
@@ -610,14 +692,16 @@ public:
 	* NOTE: This function overrides the user's previously saved seated zero pose and should only be called as the result of a user action.
 	* Users are also able to set their seated zero pose via the SteamVR Dashboard.
 	*/
-	UFUNCTION(BlueprintCallable, Category = "SteamVR Input")
+	UE_DEPRECATED(5.1, "SteamVR plugin is deprecated; please use the OpenXR plugin.")
+	UFUNCTION(BlueprintCallable, Category = "SteamVR Input", meta = (DeprecatedFunction, DeprecationMessage = "The SteamVR plugin is deprecated."))
 	static bool ResetSeatedPosition();
 
 	/**
 	* Returns the user's HMD's current IPD (interpupillary distance) setting in millimetres.
 	* @return float - The current IPD setting of the user's headset in millimetres
 	*/
-	UFUNCTION(BlueprintCallable, Category = "SteamVR Input")
+	UE_DEPRECATED(5.1, "SteamVR plugin is deprecated; please use the OpenXR plugin.")
+	UFUNCTION(BlueprintCallable, Category = "SteamVR Input", meta = (DeprecatedFunction, DeprecationMessage = "The SteamVR plugin is deprecated."))
 	static float GetUserIPD();
 
 	/**
@@ -642,3 +726,5 @@ private:
 	/** Retrieve the first available SteamVR Input device currently active in a game */
 	static class FSteamVRInputDevice* GetSteamVRInputDevice();
 };
+
+PRAGMA_ENABLE_DEPRECATION_WARNINGS
diff --git a/Engine/Plugins/Runtime/Steam/SteamVR/Source/SteamVRInputDevice/Public/SteamVRInputPublic.h b/Engine/Plugins/Runtime/Steam/SteamVR/Source/SteamVRInputDevice/Public/SteamVRInputPublic.h
--- a/Engine/Plugins/Runtime/Steam/SteamVR/Source/SteamVRInputDevice/Public/SteamVRInputPublic.h
+++ b/Engine/Plugins/Runtime/Steam/SteamVR/Source/SteamVRInputDevice/Public/SteamVRInputPublic.h
@@ -30,6 +30,8 @@ POSSIBILITY OF SUCH DAMAGE.
 
 #pragma once
 
+struct UE_DEPRECATED(5.1, "SteamVR plugin is deprecated; please use the OpenXR plugin.") FSteamVRSkeletonTransform_t;
+
 struct FSteamVRSkeletonTransform_t
 {
 	FTransform Wrist;
diff --git a/Engine/Plugins/Runtime/Steam/SteamVR/Source/SteamVRInputDevice/Public/SteamVRSkeletonDefinition.h b/Engine/Plugins/Runtime/Steam/SteamVR/Source/SteamVRInputDevice/Public/SteamVRSkeletonDefinition.h
--- a/Engine/Plugins/Runtime/Steam/SteamVR/Source/SteamVRInputDevice/Public/SteamVRSkeletonDefinition.h
+++ b/Engine/Plugins/Runtime/Steam/SteamVR/Source/SteamVRInputDevice/Public/SteamVRSkeletonDefinition.h
@@ -32,10 +32,12 @@ POSSIBILITY OF SUCH DAMAGE.
 
 #include "CoreMinimal.h"
 
+PRAGMA_DISABLE_DEPRECATION_WARNINGS
+
 /**
  * Enum for each bone in the SteamVR hand skeleton
 */
-enum ESteamVRBone : int8
+enum UE_DEPRECATED(5.1, "SteamVR plugin is deprecated; please use the OpenXR plugin.") ESteamVRBone : int8
 {
 	ESteamVRBone_Root = 0,
 	ESteamVRBone_Wrist,
@@ -93,3 +95,5 @@ namespace SteamVRSkeleton
 	/** Returns the index of the nth child of the given bone */
 	int32			GetChildIndex(int32 nBoneIndex, int32 nChildIndex);
 };
+
+PRAGMA_ENABLE_DEPRECATION_WARNINGS
diff --git a/Engine/Plugins/Runtime/Steam/SteamVR/Source/SteamVRInputDevice/Public/SteamVRTrackingRefComponent.h b/Engine/Plugins/Runtime/Steam/SteamVR/Source/SteamVRInputDevice/Public/SteamVRTrackingRefComponent.h
--- a/Engine/Plugins/Runtime/Steam/SteamVR/Source/SteamVRInputDevice/Public/SteamVRTrackingRefComponent.h
+++ b/Engine/Plugins/Runtime/Steam/SteamVR/Source/SteamVRInputDevice/Public/SteamVRTrackingRefComponent.h
@@ -39,42 +39,49 @@ POSSIBILITY OF SUCH DAMAGE.
 DECLARE_DYNAMIC_MULTICAST_DELEGATE_ThreeParams(FComponentTrackingActivatedSignature, int32, DeviceID, FName, DeviceClass, FString, DeviceModel);
 DECLARE_DYNAMIC_MULTICAST_DELEGATE_ThreeParams(FComponentTrackingDeactivatedSignature, int32, DeviceID, FName, DeviceClass, FString, DeviceModel);
 
-UCLASS( ClassGroup=(Custom), meta=(BlueprintSpawnableComponent) )
-class STEAMVRINPUTDEVICE_API USteamVRTrackingReferences : public UActorComponent
+UCLASS( ClassGroup=(Custom), deprecated, meta=(BlueprintSpawnableComponent, DeprecationMessage = "SteamVR plugin is deprecated; please use the OpenXR plugin.") )
+class STEAMVRINPUTDEVICE_API UDEPRECATED_USteamVRTrackingReferences : public UActorComponent
 {
 	GENERATED_BODY()
 
 public:	
-	USteamVRTrackingReferences();
+	UDEPRECATED_USteamVRTrackingReferences();
 
 	/** Blueprint event - When a new active device is recognized */
-	UPROPERTY(BlueprintAssignable, Category = "VR")
+	UE_DEPRECATED(5.1, "SteamVR plugin is deprecated; please use the OpenXR plugin.")
+	UPROPERTY(BlueprintAssignable, Category = "VR", meta = (DeprecatedProperty))
 	FComponentTrackingActivatedSignature OnTrackedDeviceActivated;
 
 	/** When an active device gets deactivated */
-	UPROPERTY(BlueprintAssignable, Category = "VR")
+	UE_DEPRECATED(5.1, "SteamVR plugin is deprecated; please use the OpenXR plugin.")
+	UPROPERTY(BlueprintAssignable, Category = "VR", meta = (DeprecatedProperty))
 	FComponentTrackingDeactivatedSignature OnTrackedDeviceDeactivated;
 
 	// TODO: Set default mesh to SteamVR provided render model. Must be backwards compatible to UE4.15
 
 	/** Display Tracking References in-world */
-	UFUNCTION(BlueprintCallable, Category = "SteamVR Input")
+	UE_DEPRECATED(5.1, "SteamVR plugin is deprecated; please use the OpenXR plugin.")
+	UFUNCTION(BlueprintCallable, Category = "SteamVR Input", meta = (DeprecatedFunction, DeprecationMessage = "The SteamVR plugin is deprecated."))
 	bool ShowTrackingReferences(UStaticMesh* TrackingReferenceMesh);
 
 	/** Remove Tracking References in-world */
-	UFUNCTION(BlueprintCallable, Category = "SteamVR Input")
+	UE_DEPRECATED(5.1, "SteamVR plugin is deprecated; please use the OpenXR plugin.")
+	UFUNCTION(BlueprintCallable, Category = "SteamVR Input", meta = (DeprecatedFunction, DeprecationMessage = "The SteamVR plugin is deprecated."))
 	void HideTrackingReferences();
 
 	/** Scale to apply to the tracking reference mesh */
-	UPROPERTY(EditAnywhere, BlueprintReadWrite, Category = "SteamVR Input")
+	UE_DEPRECATED(5.1, "SteamVR plugin is deprecated; please use the OpenXR plugin.")
+	UPROPERTY(EditAnywhere, BlueprintReadWrite, Category = "SteamVR Input", meta = (DeprecatedProperty))
 	float ActiveDevicePollFrequency = 1.f;
 
 	/** Scale to apply to the tracking reference mesh */
-	UPROPERTY(EditAnywhere, BlueprintReadWrite, Category = "SteamVR Input")
+	UE_DEPRECATED(5.1, "SteamVR plugin is deprecated; please use the OpenXR plugin.")
+	UPROPERTY(EditAnywhere, BlueprintReadWrite, Category = "SteamVR Input", meta = (DeprecatedProperty))
 	FVector TrackingReferenceScale = FVector(1.f);
 
 	/** Currently displayed Tracking References in-world */
-	UPROPERTY(BlueprintReadOnly, Category = "SteamVR Input")
+	UE_DEPRECATED(5.1, "SteamVR plugin is deprecated; please use the OpenXR plugin.")
+	UPROPERTY(BlueprintReadOnly, Category = "SteamVR Input", meta = (DeprecatedProperty))
 	TArray<TObjectPtr<UStaticMeshComponent>> TrackingReferences;
 
 private:
