SteamVR deprecation: SteamVR module

diff --git a/Engine/Plugins/Runtime/Steam/SteamVR/Source/SteamVR/Private/SteamVRHMD.cpp b/Engine/Plugins/Runtime/Steam/SteamVR/Source/SteamVR/Private/SteamVRHMD.cpp
--- a/Engine/Plugins/Runtime/Steam/SteamVR/Source/SteamVR/Private/SteamVRHMD.cpp
+++ b/Engine/Plugins/Runtime/Steam/SteamVR/Source/SteamVR/Private/SteamVRHMD.cpp
@@ -39,6 +39,8 @@
 
 DEFINE_LOG_CATEGORY(LogSteamVR);
 
+PRAGMA_DISABLE_DEPRECATION_WARNINGS
+
 // Visibility mesh
 static TAutoConsoleVariable<int32> CUseSteamVRVisibleAreaMesh(TEXT("vr.SteamVR.UseVisibleAreaMesh"), 1, TEXT("If non-zero, SteamVR will use the visible area mesh in addition to the hidden area mesh optimization.  This may be problematic on beta versions of platforms."));
 
@@ -1993,5 +1995,7 @@ const FSteamVRHMD::FTrackingFrame& FSteamVRHMD::GetTrackingFrame() const
 	}
 }
 
+PRAGMA_ENABLE_DEPRECATION_WARNINGS
+
 #endif //STEAMVR_SUPPORTED_PLATFORMS
 
diff --git a/Engine/Plugins/Runtime/Steam/SteamVR/Source/SteamVR/Private/SteamVRHMD.h b/Engine/Plugins/Runtime/Steam/SteamVR/Source/SteamVR/Private/SteamVRHMD.h
--- a/Engine/Plugins/Runtime/Steam/SteamVR/Source/SteamVR/Private/SteamVRHMD.h
+++ b/Engine/Plugins/Runtime/Steam/SteamVR/Source/SteamVR/Private/SteamVRHMD.h
@@ -5,6 +5,8 @@
 
 #if STEAMVR_SUPPORTED_PLATFORMS
 
+PRAGMA_DISABLE_DEPRECATION_WARNINGS
+
 #include "HeadMountedDisplay.h"
 #include "HeadMountedDisplayBase.h"
 #include "SteamVRFunctionLibrary.h"
@@ -593,5 +595,6 @@ public:
 	friend class FSteamSplashTicker;
 };
 
+PRAGMA_ENABLE_DEPRECATION_WARNINGS
 
 #endif //STEAMVR_SUPPORTED_PLATFORMS
diff --git a/Engine/Plugins/Runtime/Steam/SteamVR/Source/SteamVR/Private/SteamVRHQLayerShape.cpp b/Engine/Plugins/Runtime/Steam/SteamVR/Source/SteamVR/Private/SteamVRHQLayerShape.cpp
--- a/Engine/Plugins/Runtime/Steam/SteamVR/Source/SteamVR/Private/SteamVRHQLayerShape.cpp
+++ b/Engine/Plugins/Runtime/Steam/SteamVR/Source/SteamVR/Private/SteamVRHQLayerShape.cpp
@@ -4,6 +4,8 @@
 #include "SteamVRStereoLayers.h"
 #include "IStereoLayers.h"
 
+PRAGMA_DISABLE_DEPRECATION_WARNINGS
+
 void USteamVRHQStereoLayerShape::ApplyShape(IStereoLayers::FLayerDesc& LayerDesc)
 {
 	LayerDesc.SetShape<FSteamVRHQLayer>(bCurved, bAntiAlias, AutoCurveMinDistance, AutoCurveMaxDistance);
@@ -49,3 +51,4 @@ void USteamVRHQStereoLayerShape::SetAutoCurveMaxDistance(float InDistance)
 	MarkStereoLayerDirty();
 }
 
+PRAGMA_ENABLE_DEPRECATION_WARNINGS
diff --git a/Engine/Plugins/Runtime/Steam/SteamVR/Source/SteamVR/Private/SteamVRStereoLayers.cpp b/Engine/Plugins/Runtime/Steam/SteamVR/Source/SteamVR/Private/SteamVRStereoLayers.cpp
--- a/Engine/Plugins/Runtime/Steam/SteamVR/Source/SteamVR/Private/SteamVRStereoLayers.cpp
+++ b/Engine/Plugins/Runtime/Steam/SteamVR/Source/SteamVR/Private/SteamVRStereoLayers.cpp
@@ -12,6 +12,8 @@ const FName FSteamVRHQLayer::ShapeName = FName("SteamVRHQLayer");
 
 #if STEAMVR_SUPPORTED_PLATFORMS
 
+PRAGMA_DISABLE_DEPRECATION_WARNINGS
+
 // enable to ensure all openvr calls succeed
 #if 0
 static vr::EVROverlayError sOvrError;
@@ -361,4 +363,6 @@ IStereoLayers* FSteamVRHMD::GetStereoLayers()
 	return nullptr;
 }
 
+PRAGMA_ENABLE_DEPRECATION_WARNINGS
+
 #endif //STEAMVR_SUPPORTED_PLATFORMS
diff --git a/Engine/Plugins/Runtime/Steam/SteamVR/Source/SteamVR/Public/ISteamVRPlugin.h b/Engine/Plugins/Runtime/Steam/SteamVR/Source/SteamVR/Public/ISteamVRPlugin.h
--- a/Engine/Plugins/Runtime/Steam/SteamVR/Source/SteamVR/Public/ISteamVRPlugin.h
+++ b/Engine/Plugins/Runtime/Steam/SteamVR/Source/SteamVR/Public/ISteamVRPlugin.h
@@ -21,6 +21,7 @@ DECLARE_LOG_CATEGORY_EXTERN(LogSteamVR, Log, All);
  * The public interface to this module.  In most cases, this interface is only public to sibling modules 
  * within this plugin.
  */
+class UE_DEPRECATED(5.1, "SteamVR plugin is deprecated; please use the OpenXR plugin.") ISteamVRPlugin;
 class ISteamVRPlugin : public IHeadMountedDisplayModule
 {
 
diff --git a/Engine/Plugins/Runtime/Steam/SteamVR/Source/SteamVR/Public/SteamVRStereoLayers.h b/Engine/Plugins/Runtime/Steam/SteamVR/Source/SteamVR/Public/SteamVRStereoLayers.h
--- a/Engine/Plugins/Runtime/Steam/SteamVR/Source/SteamVR/Public/SteamVRStereoLayers.h
+++ b/Engine/Plugins/Runtime/Steam/SteamVR/Source/SteamVR/Public/SteamVRStereoLayers.h
@@ -3,6 +3,7 @@
 #include "StereoLayerShapes.h"
 
 /** Class describing additional settings for high-quality layers. Note there can only be one HQ layer active at any given time. */
+class UE_DEPRECATED(5.1, "SteamVR plugin is deprecated; please use the OpenXR plugin.") FSteamVRHQLayer;
 class STEAMVR_API FSteamVRHQLayer : public IStereoLayerShape
 {
 	STEREO_LAYER_SHAPE_BOILERPLATE(FSteamVRHQLayer)
diff --git a/Engine/Plugins/Runtime/Steam/SteamVR/Source/SteamVRInputDevice/Private/SteamVRInputDevice.h b/Engine/Plugins/Runtime/Steam/SteamVR/Source/SteamVRInputDevice/Private/SteamVRInputDevice.h
--- a/Engine/Plugins/Runtime/Steam/SteamVR/Source/SteamVRInputDevice/Private/SteamVRInputDevice.h
+++ b/Engine/Plugins/Runtime/Steam/SteamVR/Source/SteamVRInputDevice/Private/SteamVRInputDevice.h
@@ -44,6 +44,8 @@ POSSIBILITY OF SUCH DAMAGE.
 #include "SteamVRInputPublic.h"
 #include "ISteamVRPlugin.h"
 
+PRAGMA_DISABLE_DEPRECATION_WARNINGS
+
 class STEAMVRINPUTDEVICE_API FSteamVRInputDevice : public IInputDevice, public FXRMotionControllerBase, public IHapticDevice
 {
 public:
@@ -556,3 +558,5 @@ private:
 	*/
 	bool bSteamVRWasShutdown = false;
 };
+
+PRAGMA_ENABLE_DEPRECATION_WARNINGS
