SteamVR deprecation: SteamVR plugin

diff --git a/Engine/Plugins/Runtime/Steam/SteamVR/SteamVR.uplugin b/Engine/Plugins/Runtime/Steam/SteamVR/SteamVR.uplugin
--- a/Engine/Plugins/Runtime/Steam/SteamVR/SteamVR.uplugin
+++ b/Engine/Plugins/Runtime/Steam/SteamVR/SteamVR.uplugin
@@ -3,7 +3,7 @@
 	"Version": 1,
 	"VersionName": "1.0",
 	"FriendlyName": "SteamVR",
-	"Description": "Support for the SteamVR API",
+	"Description": "Support for the SteamVR API. This plugin is deprecated and will be removed in a future engine release. Please use the OpenXR plugin.",
 	"Category": "Virtual Reality",
 	"CreatedBy": "Epic Games, Inc.",
 	"CreatedByURL": "http://epicgames.com",
