Fixed SteamVR crash when calling CalculateRenderTargetSize on the render thread.

SteamVRHMD's PixelDensity variable is now pipelined within FTrackingFrame, w/ appropriate getter & setter. This makes SteamVRHMD::CalculateRenderTargetSize safe to use on the render thread, matching up with the way OpenXR handles it.

diff --git a/Engine/Plugins/Runtime/Steam/SteamVR/Source/SteamVR/Private/SteamVRHMD.cpp b/Engine/Plugins/Runtime/Steam/SteamVR/Source/SteamVR/Private/SteamVRHMD.cpp
--- a/Engine/Plugins/Runtime/Steam/SteamVR/Source/SteamVR/Private/SteamVRHMD.cpp
+++ b/Engine/Plugins/Runtime/Steam/SteamVR/Source/SteamVR/Private/SteamVRHMD.cpp
@@ -1433,6 +1433,8 @@ bool FSteamVRHMD::EnableStereo(bool bStereo)
 
 void FSteamVRHMD::AdjustViewRect(int32 ViewIndex, int32& X, int32& Y, uint32& SizeX, uint32& SizeY) const
 {
+	const float PixelDensity = GetPixelDenity();
+
 	SizeX = FMath::CeilToInt(IdealRenderTargetSize.X * PixelDensity);
 	SizeY = FMath::CeilToInt(IdealRenderTargetSize.Y * PixelDensity);
 
@@ -1544,6 +1546,24 @@ bool FSteamVRHMD::GetHMDDistortionEnabled(EShadingPath /* ShadingPath */) const
 	return false;
 }
 
+float FSteamVRHMD::GetPixelDenity() const
+{
+	const FTrackingFrame& TrackingFrame = GetTrackingFrame();
+	return TrackingFrame.PixelDensity;
+}
+
+void FSteamVRHMD::SetPixelDensity(const float NewDensity)
+{
+	check(IsInGameThread());
+	GameTrackingFrame.PixelDensity = NewDensity;
+
+	ENQUEUE_RENDER_COMMAND(UpdatePixelDensity)(
+		[this, PixelDensity = GameTrackingFrame.PixelDensity](FRHICommandListImmediate&)
+		{
+			RenderTrackingFrame.PixelDensity = PixelDensity;
+		});
+}
+
 void FSteamVRHMD::OnBeginRendering_GameThread()
 {
 	check(IsInGameThread());
@@ -1576,13 +1596,13 @@ FXRRenderBridge* FSteamVRHMD::GetActiveRenderBridge_GameThread(bool /* bUseSepar
 
 void FSteamVRHMD::CalculateRenderTargetSize(const class FViewport& Viewport, uint32& InOutSizeX, uint32& InOutSizeY)
 {
-	check(IsInGameThread());
-
 	if (!IsStereoEnabled())
 	{
 		return;
 	}
 
+	const float PixelDensity = GetPixelDenity();
+
 	InOutSizeX = FMath::CeilToInt(IdealRenderTargetSize.X * PixelDensity);
 	InOutSizeY = FMath::CeilToInt(IdealRenderTargetSize.Y * PixelDensity);
 
@@ -1758,7 +1778,6 @@ FSteamVRHMD::FSteamVRHMD(const FAutoRegister& AutoRegister, ISteamVRPlugin* InSt
 	bOcclusionMeshesBuilt(false),
 	WindowMirrorBoundsWidth(2160),
 	WindowMirrorBoundsHeight(1200),
-	PixelDensity(1.0f),
 	HMDWornMovementThreshold(50.0f),
 	HMDStartLocation(FVector::ZeroVector),
 	BaseOrientation(FQuat::Identity),
@@ -1818,7 +1837,7 @@ bool FSteamVRHMD::Startup()
 		static const auto PixelDensityCVar = IConsoleManager::Get().FindConsoleVariable(TEXT("vr.PixelDensity"));
 		if (PixelDensityCVar)
 		{ 
-			PixelDensity = FMath::Clamp(PixelDensityCVar->GetFloat(), PixelDensityMin, PixelDensityMax);
+			SetPixelDensity(FMath::Clamp(PixelDensityCVar->GetFloat(), PixelDensityMin, PixelDensityMax));
 		}
 
 		// enforce finishcurrentframe
diff --git a/Engine/Plugins/Runtime/Steam/SteamVR/Source/SteamVR/Private/SteamVRHMD.h b/Engine/Plugins/Runtime/Steam/SteamVR/Source/SteamVR/Private/SteamVRHMD.h
--- a/Engine/Plugins/Runtime/Steam/SteamVR/Source/SteamVR/Private/SteamVRHMD.h
+++ b/Engine/Plugins/Runtime/Steam/SteamVR/Source/SteamVR/Private/SteamVRHMD.h
@@ -158,8 +158,8 @@ public:
 	virtual void OnBeginRendering_GameThread() override;
 	virtual void OnBeginRendering_RenderThread(FRHICommandListImmediate& RHICmdList, FSceneViewFamily& ViewFamily) override;
 
-	virtual float GetPixelDenity() const override { return PixelDensity; }
-	virtual void SetPixelDensity(const float NewDensity) override { PixelDensity = NewDensity; }
+	virtual float GetPixelDenity() const override;
+	virtual void SetPixelDensity(const float NewDensity) override;
 	virtual FIntPoint GetIdealRenderTargetSize() const override { return IdealRenderTargetSize; }
 
 	/** IStereoRendering interface */
@@ -475,6 +475,7 @@ private:
 
 		/** World units (UU) to Meters scale.  Read from the level, and used to transform positional tracking data */
 		float WorldToMetersScale;
+		float PixelDensity;
 
 		vr::HmdMatrix34_t RawPoses[vr::k_unMaxTrackedDeviceCount];
 
@@ -482,6 +483,7 @@ private:
 			: FrameNumber(0)
 			, bHaveVisionTracking(false)
 			, WorldToMetersScale(100.0f)
+			, PixelDensity(1.0f)
 		{
 			const uint32 MaxDevices = vr::k_unMaxTrackedDeviceCount;
 
@@ -556,7 +558,6 @@ private:
 	uint32 WindowMirrorBoundsHeight;
 
 	FIntPoint IdealRenderTargetSize;
-	float PixelDensity;
 
 	/** How far the HMD has to move before it's considered to be worn */
 	float HMDWornMovementThreshold;
