Remove full path to public Core headers

#rnx

diff --git a/Engine/Plugins/Developer/XCodeSourceCodeAccess/Source/XCodeSourceCodeAccess/Private/XCodeSourceCodeAccessModule.cpp b/Engine/Plugins/Developer/XCodeSourceCodeAccess/Source/XCodeSourceCodeAccess/Private/XCodeSourceCodeAccessModule.cpp
--- a/Engine/Plugins/Developer/XCodeSourceCodeAccess/Source/XCodeSourceCodeAccess/Private/XCodeSourceCodeAccessModule.cpp
+++ b/Engine/Plugins/Developer/XCodeSourceCodeAccess/Source/XCodeSourceCodeAccess/Private/XCodeSourceCodeAccessModule.cpp
@@ -1,7 +1,7 @@
 // Copyright Epic Games, Inc. All Rights Reserved.
 
 #include "XCodeSourceCodeAccessModule.h"
-#include "Runtime/Core/Public/Features/IModularFeatures.h"
+#include "Features/IModularFeatures.h"
 #include "Modules/ModuleManager.h"
 
 IMPLEMENT_MODULE( FXCodeSourceCodeAccessModule, XCodeSourceCodeAccess );
diff --git a/Engine/Plugins/Editor/MacGraphicsSwitching/Source/MacGraphicsSwitching/Private/MacGraphicsSwitchingModule.cpp b/Engine/Plugins/Editor/MacGraphicsSwitching/Source/MacGraphicsSwitching/Private/MacGraphicsSwitchingModule.cpp
--- a/Engine/Plugins/Editor/MacGraphicsSwitching/Source/MacGraphicsSwitching/Private/MacGraphicsSwitchingModule.cpp
+++ b/Engine/Plugins/Editor/MacGraphicsSwitching/Source/MacGraphicsSwitching/Private/MacGraphicsSwitchingModule.cpp
@@ -6,7 +6,7 @@
 #include "MacGraphicsSwitchingSettingsDetails.h"
 #include "MacGraphicsSwitchingStyle.h"
 #include "MacGraphicsSwitchingWidget.h"
-#include "Runtime/Core/Public/Features/IModularFeatures.h"
+#include "Features/IModularFeatures.h"
 #include "Runtime/SlateCore/Public/Rendering/SlateRenderer.h"
 #include "Editor/MainFrame/Public/Interfaces/IMainFrameModule.h"
 #include "ISettingsModule.h"
diff --git a/Engine/Plugins/Media/BinkMedia/Source/BinkMediaPlayer/Private/BinkMediaPlayerModule.cpp b/Engine/Plugins/Media/BinkMedia/Source/BinkMediaPlayer/Private/BinkMediaPlayerModule.cpp
--- a/Engine/Plugins/Media/BinkMedia/Source/BinkMediaPlayer/Private/BinkMediaPlayerModule.cpp
+++ b/Engine/Plugins/Media/BinkMedia/Source/BinkMediaPlayer/Private/BinkMediaPlayerModule.cpp
@@ -5,7 +5,7 @@
 #include "BinkMediaPlayerPrivate.h"
 #include "Engine/GameViewportClient.h"
 #include "Internationalization/Internationalization.h"
-#include "Runtime/Core/Public/Misc/Paths.h"
+#include "Misc/Paths.h"
 #include "Interfaces/IPluginManager.h"
 #include "Misc/CoreDelegates.h"
 #include "UObject/UObjectIterator.h"
diff --git a/Engine/Plugins/Runtime/Oculus/OculusVR/Source/OculusEditor/Private/OculusBuildAnalytics.cpp b/Engine/Plugins/Runtime/Oculus/OculusVR/Source/OculusEditor/Private/OculusBuildAnalytics.cpp
--- a/Engine/Plugins/Runtime/Oculus/OculusVR/Source/OculusEditor/Private/OculusBuildAnalytics.cpp
+++ b/Engine/Plugins/Runtime/Oculus/OculusVR/Source/OculusEditor/Private/OculusBuildAnalytics.cpp
@@ -3,7 +3,7 @@
 #include "OculusBuildAnalytics.h"
 #include "GameProjectGenerationModule.h"
 #include "OculusHMDModule.h"
-#include "Runtime/Core/Public/HAL/FileManager.h"
+#include "HAL/FileManager.h"
 
 FOculusBuildAnalytics* FOculusBuildAnalytics::instance = 0;
 
diff --git a/Engine/Plugins/Runtime/Steam/SteamVR/Source/SteamVREditor/Private/SteamVREditor.cpp b/Engine/Plugins/Runtime/Steam/SteamVR/Source/SteamVREditor/Private/SteamVREditor.cpp
--- a/Engine/Plugins/Runtime/Steam/SteamVR/Source/SteamVREditor/Private/SteamVREditor.cpp
+++ b/Engine/Plugins/Runtime/Steam/SteamVR/Source/SteamVREditor/Private/SteamVREditor.cpp
@@ -42,8 +42,8 @@ POSSIBILITY OF SUCH DAMAGE.
 #include "DetailLayoutBuilder.h"
 #include "Widgets/Input/SButton.h"
 #include "DetailWidgetRow.h"
-#include "Runtime/Core/Public/CoreGlobals.h"
-#include "Runtime/Core/Public/Misc/ConfigCacheIni.h"
+#include "CoreGlobals.h"
+#include "Misc/ConfigCacheIni.h"
 #include "ToolMenus.h"
 
 static const FName SteamVREditorTabName("SteamVREditor");
diff --git a/Engine/Plugins/Runtime/Steam/SteamVR/Source/SteamVRInput/Private/SteamVRInput.cpp b/Engine/Plugins/Runtime/Steam/SteamVR/Source/SteamVRInput/Private/SteamVRInput.cpp
--- a/Engine/Plugins/Runtime/Steam/SteamVR/Source/SteamVRInput/Private/SteamVRInput.cpp
+++ b/Engine/Plugins/Runtime/Steam/SteamVR/Source/SteamVRInput/Private/SteamVRInput.cpp
@@ -30,7 +30,7 @@ POSSIBILITY OF SUCH DAMAGE.
 
 #include "SteamVRInput.h"
 #include "CoreMinimal.h"
-#include "Runtime/Core/Public/Misc/Paths.h"
+#include "Misc/Paths.h"
 #include "openvr.h"
 #include "Modules/ModuleManager.h"
 #include "Interfaces/IPluginManager.h"
diff --git a/Engine/Plugins/Runtime/Steam/SteamVR/Source/SteamVRInputDevice/Private/SteamVRInputDevice.h b/Engine/Plugins/Runtime/Steam/SteamVR/Source/SteamVRInputDevice/Private/SteamVRInputDevice.h
--- a/Engine/Plugins/Runtime/Steam/SteamVR/Source/SteamVRInputDevice/Private/SteamVRInputDevice.h
+++ b/Engine/Plugins/Runtime/Steam/SteamVR/Source/SteamVRInputDevice/Private/SteamVRInputDevice.h
@@ -32,11 +32,11 @@ POSSIBILITY OF SUCH DAMAGE.
 
 #include "IInputDevice.h"
 #include "CoreMinimal.h"
-#include "Runtime/Core/Public/Misc/ConfigCacheIni.h"
+#include "Misc/ConfigCacheIni.h"
 #include "XRMotionControllerBase.h"
 #include "IHapticDevice.h"
-#include "Runtime/Core/Public/GenericPlatform/GenericPlatformProcess.h"
-#include "Runtime/Core/Public/Misc/Paths.h"
+#include "GenericPlatform/GenericPlatformProcess.h"
+#include "Misc/Paths.h"
 #include "Engine/Engine.h"
 #include "Serialization/JsonReader.h"
 #include "Serialization/JsonSerializer.h"
diff --git a/Engine/Plugins/Runtime/Steam/SteamVR/Source/SteamVRInputDevice/Public/SteamVRInputDeviceFunctionLibrary.h b/Engine/Plugins/Runtime/Steam/SteamVR/Source/SteamVRInputDevice/Public/SteamVRInputDeviceFunctionLibrary.h
--- a/Engine/Plugins/Runtime/Steam/SteamVR/Source/SteamVRInputDevice/Public/SteamVRInputDeviceFunctionLibrary.h
+++ b/Engine/Plugins/Runtime/Steam/SteamVR/Source/SteamVRInputDevice/Public/SteamVRInputDeviceFunctionLibrary.h
@@ -36,7 +36,7 @@ POSSIBILITY OF SUCH DAMAGE.
 #include "CoreMinimal.h"
 #include "Kismet/BlueprintFunctionLibrary.h"
 #include "Features/IModularFeatures.h"
-#include "Runtime/Core/Public/GenericPlatform/GenericPlatformProcess.h"
+#include "GenericPlatform/GenericPlatformProcess.h"
 #include "IMotionController.h"
 #include "Engine/World.h"
 #include "Kismet/GameplayStatics.h"
diff --git a/Engine/Plugins/VirtualProduction/MultiUserTakes/Source/ConcertTakeRecorder/Private/ConcertTakeRecorderClientSessionCustomization.cpp b/Engine/Plugins/VirtualProduction/MultiUserTakes/Source/ConcertTakeRecorder/Private/ConcertTakeRecorderClientSessionCustomization.cpp
--- a/Engine/Plugins/VirtualProduction/MultiUserTakes/Source/ConcertTakeRecorder/Private/ConcertTakeRecorderClientSessionCustomization.cpp
+++ b/Engine/Plugins/VirtualProduction/MultiUserTakes/Source/ConcertTakeRecorder/Private/ConcertTakeRecorderClientSessionCustomization.cpp
@@ -33,7 +33,7 @@
 #include "Styling/AppStyle.h"
 #include "DetailWidgetRow.h"
 
-#include "Core/Public/Modules/ModuleManager.h"
+#include "Modules/ModuleManager.h"
 #include "ISettingsModule.h"
 
 #define LOCTEXT_NAMESPACE "ConcertTakeRecorderClientSessionCustomization"
diff --git a/Engine/Plugins/VirtualProduction/MultiUserTakes/Source/ConcertTakeRecorder/Private/ConcertTakeRecorderSynchronizationCustomization.h b/Engine/Plugins/VirtualProduction/MultiUserTakes/Source/ConcertTakeRecorder/Private/ConcertTakeRecorderSynchronizationCustomization.h
--- a/Engine/Plugins/VirtualProduction/MultiUserTakes/Source/ConcertTakeRecorder/Private/ConcertTakeRecorderSynchronizationCustomization.h
+++ b/Engine/Plugins/VirtualProduction/MultiUserTakes/Source/ConcertTakeRecorder/Private/ConcertTakeRecorderSynchronizationCustomization.h
@@ -19,7 +19,7 @@
 #include "Widgets/Input/SHyperlink.h"
 #include "DetailWidgetRow.h"
 
-#include "Core/Public/Modules/ModuleManager.h"
+#include "Modules/ModuleManager.h"
 #include "ISettingsModule.h"
 
 #define LOCTEXT_NAMESPACE "ConcertTakeRecorderCustomization"
diff --git a/Engine/Source/Programs/UnrealLightmass/Private/Launch/UnrealLightmass.cpp b/Engine/Source/Programs/UnrealLightmass/Private/Launch/UnrealLightmass.cpp
--- a/Engine/Source/Programs/UnrealLightmass/Private/Launch/UnrealLightmass.cpp
+++ b/Engine/Source/Programs/UnrealLightmass/Private/Launch/UnrealLightmass.cpp
@@ -7,7 +7,7 @@
 #include "CPUSolver.h"
 #include "UnitTest.h"
 #include "LightmassSwarm.h"
-#include "Runtime/Core/Public/HAL/ExceptionHandling.h"
+#include "HAL/ExceptionHandling.h"
 #include "RequiredProgramMainCPPInclude.h"
 #include "LMDebug.h"
 #include "LMHelpers.h"
diff --git a/Engine/Source/Programs/UnrealVersionSelector/Private/Windows/WindowsPlatformInstallation.cpp b/Engine/Source/Programs/UnrealVersionSelector/Private/Windows/WindowsPlatformInstallation.cpp
--- a/Engine/Source/Programs/UnrealVersionSelector/Private/Windows/WindowsPlatformInstallation.cpp
+++ b/Engine/Source/Programs/UnrealVersionSelector/Private/Windows/WindowsPlatformInstallation.cpp
@@ -1,7 +1,7 @@
 // Copyright Epic Games, Inc. All Rights Reserved.
 
 #include "WindowsPlatformInstallation.h"
-#include "Runtime/Core/Public/Misc/EngineVersion.h"
+#include "Misc/EngineVersion.h"
 #include "Windows/AllowWindowsPlatformTypes.h"
 #include "DesktopPlatformModule.h"
 #include "Misc/Paths.h"
diff --git a/Engine/Source/Runtime/Core/Core.Build.cs b/Engine/Source/Runtime/Core/Core.Build.cs
--- a/Engine/Source/Runtime/Core/Core.Build.cs
+++ b/Engine/Source/Runtime/Core/Core.Build.cs
@@ -23,17 +23,11 @@ public class Core : ModuleRules
 		PrivateDependencyModuleNames.Add("OodleDataCompression");
 
 		PublicDependencyModuleNames.Add("TraceLog");
-		PublicIncludePaths.Add("Runtime/TraceLog/Public");
-
-		PrivateIncludePaths.AddRange(
-			new string[] {
-				"Runtime/SynthBenchmark/Public",
-				"Runtime/Engine/Public",
-			}
-			);
 
 		PrivateIncludePathModuleNames.AddRange(
 			new string[] {
+				"SynthBenchmark",
+				"Engine",
 				"TargetPlatform",
 				"DerivedDataCache",
 				"InputDevice",
